MovieSceneCapture build fix for CIS. Fixes errors in non-editor builds of UE4 from CL 2700384.

diff --git a/Engine/Source/Runtime/ActorAnimation/Private/ActorAnimationObjectReference.cpp b/Engine/Source/Runtime/ActorAnimation/Private/ActorAnimationObjectReference.cpp
--- a/Engine/Source/Runtime/ActorAnimation/Private/ActorAnimationObjectReference.cpp
+++ b/Engine/Source/Runtime/ActorAnimation/Private/ActorAnimationObjectReference.cpp
@@ -14,8 +14,13 @@ FActorAnimationObjectReference::FActorAnimationObjectReference(UObject* InObject
 
 UObject* FActorAnimationObjectReference::Resolve(UObject* InContext) const
 {
+#if WITH_EDITOR
 	int32 PIEInstanceID = InContext->GetOutermost()->PIEInstanceID;
 	FUniqueObjectGuid FixedUpId = PIEInstanceID == -1 ? ObjectId : ObjectId.FixupForPIE(PIEInstanceID);
+#else
+	FUniqueObjectGuid FixedUpId = ObjectId;
+#endif
+	
 
 	FLazyObjectPtr LazyPtr;
 	LazyPtr = FixedUpId;
diff --git a/Engine/Source/Runtime/MovieSceneCapture/MovieSceneCapture.Build.cs b/Engine/Source/Runtime/MovieSceneCapture/MovieSceneCapture.Build.cs
--- a/Engine/Source/Runtime/MovieSceneCapture/MovieSceneCapture.Build.cs
+++ b/Engine/Source/Runtime/MovieSceneCapture/MovieSceneCapture.Build.cs
@@ -12,12 +12,17 @@ public class MovieSceneCapture : ModuleRules
 			}
 		);
 
-		var ExternalModules = new string[] {
-			"PropertyEditor",
-			"ImageWrapper",
-		};
-		PrivateIncludePathModuleNames.AddRange(ExternalModules);
-		DynamicallyLoadedModuleNames.AddRange(ExternalModules);
+		if (UEBuildConfiguration.bBuildDeveloperTools)
+		{
+			PrivateIncludePathModuleNames.Add("ImageWrapper");
+			DynamicallyLoadedModuleNames.Add("ImageWrapper");
+		}
+
+		if (UEBuildConfiguration.bBuildEditor)
+		{
+			PrivateIncludePathModuleNames.Add("PropertyEditor");
+			DynamicallyLoadedModuleNames.Add("PropertyEditor");
+		}
 
 		PublicDependencyModuleNames.AddRange(
 			new string[] {
@@ -29,9 +34,7 @@ public class MovieSceneCapture : ModuleRules
 			new string[] {
 				"Core",
 				"CoreUObject",
-				"EditorStyle",
 				"Engine",
-				"UnrealEd",
 				"InputCore",
 				"Json",
 				"JsonUtilities",
@@ -41,6 +44,14 @@ public class MovieSceneCapture : ModuleRules
 			}
 		);
 
-
+		if (UEBuildConfiguration.bBuildEditor)
+		{
+			PrivateDependencyModuleNames.AddRange(
+				new string[] {
+					"EditorStyle",
+					"UnrealEd",
+				}
+			);
+		}
 	}
 }
diff --git a/Engine/Source/Runtime/MovieSceneCapture/Private/MovieSceneCapture.cpp b/Engine/Source/Runtime/MovieSceneCapture/Private/MovieSceneCapture.cpp
--- a/Engine/Source/Runtime/MovieSceneCapture/Private/MovieSceneCapture.cpp
+++ b/Engine/Source/Runtime/MovieSceneCapture/Private/MovieSceneCapture.cpp
@@ -2,7 +2,11 @@
 
 #include "MovieSceneCapturePCH.h"
 #include "MovieSceneCapture.h"
+
+#if WITH_EDITOR
 #include "ImageWrapper.h"
+#endif
+
 #include "MovieSceneCaptureModule.h"
 
 struct FUniqueMovieSceneCaptureHandle : FMovieSceneCaptureHandle
@@ -135,12 +139,14 @@ void UMovieSceneCapture::CaptureFrame(float DeltaSeconds)
 			AVIWriter->DropFrames(NumDroppedFrames);
 			AVIWriter->Update(CurrentTimeSeconds);
 		}
+#if WITH_EDITOR
 		else
 		{
 			TArray<FColor> Data;
 			Viewport->ReadPixels(Data, FReadSurfaceDataFlags());
 			CaptureSnapshot(Data);
 		}
+#endif
 	}
 }
 
@@ -162,6 +168,7 @@ void UMovieSceneCapture::Close()
 	IMovieSceneCaptureModule::Get().OnMovieSceneCaptureFinished(this);
 }
 
+#if WITH_EDITOR
 void UMovieSceneCapture::CaptureSnapshot(const TArray<FColor>& Colors)
 {
 	if (Colors.Num() == 0)
@@ -217,6 +224,7 @@ void UMovieSceneCapture::CaptureSnapshot(const TArray<FColor>& Colors)
 		break;
 	}
 }
+#endif
 
 const TCHAR* UMovieSceneCapture::GetDefaultFileExtension() const
 {
diff --git a/Engine/Source/Runtime/MovieSceneCapture/Private/MovieSceneCaptureModule.cpp b/Engine/Source/Runtime/MovieSceneCapture/Private/MovieSceneCaptureModule.cpp
--- a/Engine/Source/Runtime/MovieSceneCapture/Private/MovieSceneCaptureModule.cpp
+++ b/Engine/Source/Runtime/MovieSceneCapture/Private/MovieSceneCaptureModule.cpp
@@ -3,10 +3,7 @@
 #include "MovieSceneCapturePCH.h"
 #include "MovieSceneCapture.h"
 #include "MovieSceneCaptureModule.h"
-#include "EditorStyle.h"
 #include "SDockTab.h"
-#include "Editor.h"
-#include "PropertyEditing.h"
 #include "JsonObjectConverter.h"
 #include "INotificationWidget.h"
 #include "SNotificationList.h"
@@ -14,105 +11,112 @@
 
 #include "SlateExtras.h"
 
+#if WITH_EDITOR
+	#include "EditorStyle.h"
+	#include "Editor.h"
+	#include "PropertyEditing.h"
+#endif
+
 #define LOCTEXT_NAMESPACE "MovieSceneCapture"
 
 const FName MovieSceneRenderTabId(TEXT("MovieSceneRenderTab"));
 
 DECLARE_DELEGATE_RetVal_OneParam(FText, FOnStartCapture, UMovieSceneCapture*);
 
-class SRenderMovieSceneSettings : public SCompoundWidget, public FGCObject
-{
-	SLATE_BEGIN_ARGS(SRenderMovieSceneSettings) : _InitialObject(nullptr) {}
-		SLATE_EVENT(FOnStartCapture, OnStartCapture)
-		SLATE_ARGUMENT(UMovieSceneCapture*, InitialObject)
-	SLATE_END_ARGS()
-
-	void Construct(const FArguments& InArgs)
+#if WITH_EDITOR
+	class SRenderMovieSceneSettings : public SCompoundWidget, public FGCObject
 	{
-		FPropertyEditorModule& PropertyEditor = FModuleManager::LoadModuleChecked<FPropertyEditorModule>("PropertyEditor");
+		SLATE_BEGIN_ARGS(SRenderMovieSceneSettings) : _InitialObject(nullptr) {}
+			SLATE_EVENT(FOnStartCapture, OnStartCapture)
+			SLATE_ARGUMENT(UMovieSceneCapture*, InitialObject)
+		SLATE_END_ARGS()
 
-		FDetailsViewArgs DetailsViewArgs;
-		DetailsViewArgs.bUpdatesFromSelection = false;
-		DetailsViewArgs.bLockable = false;
-		DetailsViewArgs.NameAreaSettings = FDetailsViewArgs::HideNameArea;
-		DetailsViewArgs.ViewIdentifier = "RenderMovieScene";
+		void Construct(const FArguments& InArgs)
+		{
+			FPropertyEditorModule& PropertyEditor = FModuleManager::LoadModuleChecked<FPropertyEditorModule>("PropertyEditor");
 
-		DetailView = PropertyEditor.CreateDetailView(DetailsViewArgs);
+			FDetailsViewArgs DetailsViewArgs;
+			DetailsViewArgs.bUpdatesFromSelection = false;
+			DetailsViewArgs.bLockable = false;
+			DetailsViewArgs.NameAreaSettings = FDetailsViewArgs::HideNameArea;
+			DetailsViewArgs.ViewIdentifier = "RenderMovieScene";
 
-		OnStartCapture = InArgs._OnStartCapture;
+			DetailView = PropertyEditor.CreateDetailView(DetailsViewArgs);
 
-		ChildSlot
-		[
-			SNew(SVerticalBox)
+			OnStartCapture = InArgs._OnStartCapture;
 
-			+ SVerticalBox::Slot()
+			ChildSlot
 			[
-				DetailView.ToSharedRef()
-			]
+				SNew(SVerticalBox)
 
-			+ SVerticalBox::Slot()
-			.AutoHeight()
-			[
-				SAssignNew(ErrorText, STextBlock)
-				.Visibility(EVisibility::Hidden)
-			]
+				+ SVerticalBox::Slot()
+				[
+					DetailView.ToSharedRef()
+				]
 
-			+ SVerticalBox::Slot()
-			.AutoHeight()
-			.HAlign(HAlign_Right)
-			.Padding(5.f)
-			[
-				SNew(SButton)
-				.ContentPadding(FMargin(10, 5))
-				.Text(LOCTEXT("Export", "Capture Movie"))
-				.OnClicked(this, &SRenderMovieSceneSettings::OnStartClicked)
-			]
+				+ SVerticalBox::Slot()
+				.AutoHeight()
+				[
+					SAssignNew(ErrorText, STextBlock)
+					.Visibility(EVisibility::Hidden)
+				]
+
+				+ SVerticalBox::Slot()
+				.AutoHeight()
+				.HAlign(HAlign_Right)
+				.Padding(5.f)
+				[
+					SNew(SButton)
+					.ContentPadding(FMargin(10, 5))
+					.Text(LOCTEXT("Export", "Capture Movie"))
+					.OnClicked(this, &SRenderMovieSceneSettings::OnStartClicked)
+				]
 			
-		];
+			];
 
-		if (InArgs._InitialObject)
-		{
-			SetObject(InArgs._InitialObject);
+			if (InArgs._InitialObject)
+			{
+				SetObject(InArgs._InitialObject);
+			}
 		}
-	}
 
-	void SetObject(UMovieSceneCapture* InMovieSceneCapture)
-	{
-		MovieSceneCapture = InMovieSceneCapture;
-
-		DetailView->SetObject(InMovieSceneCapture);
-
-		ErrorText->SetText(FText());
-		ErrorText->SetVisibility(EVisibility::Hidden);
-	}
+		void SetObject(UMovieSceneCapture* InMovieSceneCapture)
+		{
+			MovieSceneCapture = InMovieSceneCapture;
 
-	virtual void AddReferencedObjects( FReferenceCollector& Collector ) override
-	{
-		Collector.AddReferencedObject(MovieSceneCapture);
-	}
+			DetailView->SetObject(InMovieSceneCapture);
 
-private:
+			ErrorText->SetText(FText());
+			ErrorText->SetVisibility(EVisibility::Hidden);
+		}
 
-	FReply OnStartClicked()
-	{
-		FText Error;
-		if (OnStartCapture.IsBound())
+		virtual void AddReferencedObjects( FReferenceCollector& Collector ) override
 		{
-			Error = OnStartCapture.Execute(MovieSceneCapture);
+			Collector.AddReferencedObject(MovieSceneCapture);
 		}
 
-		ErrorText->SetText(Error);
-		ErrorText->SetVisibility(Error.IsEmpty() ? EVisibility::Hidden : EVisibility::Visible);
+	private:
 
-		return FReply::Handled();
-	}
+		FReply OnStartClicked()
+		{
+			FText Error;
+			if (OnStartCapture.IsBound())
+			{
+				Error = OnStartCapture.Execute(MovieSceneCapture);
+			}
 
-	TSharedPtr<IDetailsView> DetailView;
-	TSharedPtr<STextBlock> ErrorText;
-	FOnStartCapture OnStartCapture;
-	UMovieSceneCapture* MovieSceneCapture;
-};
+			ErrorText->SetText(Error);
+			ErrorText->SetVisibility(Error.IsEmpty() ? EVisibility::Hidden : EVisibility::Visible);
+
+			return FReply::Handled();
+		}
 
+		TSharedPtr<IDetailsView> DetailView;
+		TSharedPtr<STextBlock> ErrorText;
+		FOnStartCapture OnStartCapture;
+		UMovieSceneCapture* MovieSceneCapture;
+	};
+#endif
 
 
 DECLARE_DELEGATE_OneParam(FOnProcessClosed, int32);
@@ -430,6 +434,7 @@ private:
 
 	virtual void OpenCaptureSettings(const TSharedRef<FTabManager>& TabManager, UMovieSceneCapture* CaptureObject) override
 	{
+#if WITH_EDITOR
 		TSharedPtr<SWindow> ExistingWindow = CaptureSettingsWindow.Pin();
 		if (ExistingWindow.IsValid())
 		{
@@ -463,6 +468,7 @@ private:
 		);
 
 		CaptureSettingsWindow = ExistingWindow;
+#endif
 	}
 
 	void OnNotificationClosed()
diff --git a/Engine/Source/Runtime/MovieSceneCapture/Public/MovieSceneCapture.h b/Engine/Source/Runtime/MovieSceneCapture/Public/MovieSceneCapture.h
--- a/Engine/Source/Runtime/MovieSceneCapture/Public/MovieSceneCapture.h
+++ b/Engine/Source/Runtime/MovieSceneCapture/Public/MovieSceneCapture.h
@@ -69,8 +69,10 @@ public:
 
 protected:
 
+#if WITH_EDITOR
 	/** Implementation function that saves out a snapshot file from the specified color data */
 	void CaptureSnapshot(const TArray<FColor>& Colors);
+#endif
 
 protected:
 
