Added ability to retrieve the current frame number and elapsed time for a running movie scene capture

This addresses UETOOL-473.

diff --git a/Engine/Source/Runtime/Engine/Public/UnrealClient.h b/Engine/Source/Runtime/Engine/Public/UnrealClient.h
--- a/Engine/Source/Runtime/Engine/Public/UnrealClient.h
+++ b/Engine/Source/Runtime/Engine/Public/UnrealClient.h
@@ -630,7 +630,7 @@ public:
 	bool HasMovieSceneCapture() const { return MovieSceneCaptureHandle.IsValid(); }
 
 	/** Access the current movie scene capture interface, if set. */
-	IMovieSceneCaptureInterface* GetMovieSceneCapture() const;
+	ENGINE_API IMovieSceneCaptureInterface* GetMovieSceneCapture() const;
 
 	/** Set this viewport's capture implementation */
 	void SetMovieSceneCapture(FMovieSceneCaptureHandle Handle) { MovieSceneCaptureHandle = Handle; }
diff --git a/Engine/Source/Runtime/MovieSceneCapture/Private/MovieSceneCapture.cpp b/Engine/Source/Runtime/MovieSceneCapture/Private/MovieSceneCapture.cpp
--- a/Engine/Source/Runtime/MovieSceneCapture/Private/MovieSceneCapture.cpp
+++ b/Engine/Source/Runtime/MovieSceneCapture/Private/MovieSceneCapture.cpp
@@ -63,8 +63,6 @@ UMovieSceneCapture::UMovieSceneCapture(const FObjectInitializer& Initializer)
 	AdditionalCommandLineArguments += FParse::Param(FCommandLine::Get(), TEXT("opengl4"))	?	TEXT("-opengl4 ")	: TEXT("");
 
 	Handle = FUniqueMovieSceneCaptureHandle();
-
-	CurrentTimeSeconds = 0;
 }
 
 void UMovieSceneCapture::Initialize(FViewport* InViewport)
@@ -81,7 +79,7 @@ void UMovieSceneCapture::StartCapture()
 
 	CaptureStrategy->OnStart();
 
-	CurrentTimeSeconds = 0;
+	CachedMetrics.ElapsedSeconds = 0;
 
 	if (bBufferVisualizationDumpFrames)
 	{
@@ -95,6 +93,8 @@ void UMovieSceneCapture::StartCapture()
 	CachedMetrics.Width = Viewport->GetSizeXY().X;
 	CachedMetrics.Height = Viewport->GetSizeXY().Y;
 
+	Viewport->SetMovieSceneCapture(Handle);
+
 	if (Settings.CaptureType == EMovieCaptureType::AVI)
 	{
 		FAVIWriterOptions Options;
@@ -111,10 +111,6 @@ void UMovieSceneCapture::StartCapture()
 		AVIWriter.Reset(FAVIWriter::CreateInstance(Options));
 		AVIWriter->StartCapture(Viewport);
 	}
-	else
-	{
-		Viewport->SetMovieSceneCapture(Handle);
-	}
 }
 
 void UMovieSceneCapture::CaptureFrame(float DeltaSeconds)
@@ -124,20 +120,20 @@ void UMovieSceneCapture::CaptureFrame(float DeltaSeconds)
 		return;
 	}
 
-	CurrentTimeSeconds += DeltaSeconds;
+	CachedMetrics.ElapsedSeconds += DeltaSeconds;
 
-	if (CaptureStrategy->ShouldPresent(CurrentTimeSeconds, CachedMetrics.Frame))
+	if (CaptureStrategy->ShouldPresent(CachedMetrics.ElapsedSeconds, CachedMetrics.Frame))
 	{
-		uint32 NumDroppedFrames = CaptureStrategy->GetDroppedFrames(CurrentTimeSeconds, CachedMetrics.Frame);
+		uint32 NumDroppedFrames = CaptureStrategy->GetDroppedFrames(CachedMetrics.ElapsedSeconds, CachedMetrics.Frame);
 		CachedMetrics.Frame += NumDroppedFrames;
 
-		CaptureStrategy->OnPresent(CurrentTimeSeconds, CachedMetrics.Frame);
+		CaptureStrategy->OnPresent(CachedMetrics.ElapsedSeconds, CachedMetrics.Frame);
 		++CachedMetrics.Frame;
 		
 		if (AVIWriter)
 		{
 			AVIWriter->DropFrames(NumDroppedFrames);
-			AVIWriter->Update(CurrentTimeSeconds);
+			AVIWriter->Update(CachedMetrics.ElapsedSeconds);
 		}
 #if WITH_EDITOR
 		else
@@ -245,7 +241,7 @@ FString UMovieSceneCapture::ResolveFileFormat(const FString& Folder, const FStri
 {
 	TMap<FString, FStringFormatArg> Mappings;
 	Mappings.Add(TEXT("fps"), FString::Printf(TEXT("%d"), Settings.FrameRate));
-	Mappings.Add(TEXT("frame"), FString::Printf(TEXT("%04d"), CachedMetrics.Frame+1));
+	Mappings.Add(TEXT("frame"), FString::Printf(TEXT("%04d"), CachedMetrics.Frame));
 	Mappings.Add(TEXT("width"), FString::Printf(TEXT("%d"), CachedMetrics.Width));
 	Mappings.Add(TEXT("height"), FString::Printf(TEXT("%d"), CachedMetrics.Height));
 
diff --git a/Engine/Source/Runtime/MovieSceneCapture/Private/MovieSceneCaptureEnvironment.cpp b/Engine/Source/Runtime/MovieSceneCapture/Private/MovieSceneCaptureEnvironment.cpp
new file mode 100644
--- /dev/null
+++ b/Engine/Source/Runtime/MovieSceneCapture/Private/MovieSceneCaptureEnvironment.cpp
@@ -0,0 +1,35 @@
+// Copyright 1998-2015 Epic Games, Inc. All Rights Reserved.
+
+#include "MovieSceneCapturePCH.h"
+#include "MovieSceneCapture.h"
+#include "MovieSceneCaptureEnvironment.h"
+
+
+UMovieSceneCapture* FindMovieSceneCapture()
+{
+	for (const FWorldContext& Context : GEngine->GetWorldContexts())
+	{
+		if (Context.GameViewport && Context.GameViewport->Viewport)
+		{
+			UMovieSceneCapture* Capture = static_cast<UMovieSceneCapture*>(Context.GameViewport->Viewport->GetMovieSceneCapture());
+			if (Capture)
+			{
+				return Capture;
+			}
+		}
+	}
+
+	return nullptr;
+}
+
+int32 UMovieSceneCaptureEnvironment::GetCaptureFrameNumber()
+{
+	UMovieSceneCapture* Capture = FindMovieSceneCapture();
+	return Capture ? Capture->GetMetrics().Frame : 0;
+}
+
+float UMovieSceneCaptureEnvironment::GetCaptureElapsedTime()
+{
+	UMovieSceneCapture* Capture = FindMovieSceneCapture();
+	return Capture ? Capture->GetMetrics().ElapsedSeconds : 0.f;
+}
\ No newline at end of file
diff --git a/Engine/Source/Runtime/MovieSceneCapture/Public/MovieSceneCapture.h b/Engine/Source/Runtime/MovieSceneCapture/Public/MovieSceneCapture.h
--- a/Engine/Source/Runtime/MovieSceneCapture/Public/MovieSceneCapture.h
+++ b/Engine/Source/Runtime/MovieSceneCapture/Public/MovieSceneCapture.h
@@ -23,6 +23,19 @@ struct ICaptureStrategy
 	virtual int32 GetDroppedFrames(double CurrentTimeSeconds, uint32 FrameIndex) const = 0;
 };
 
+/** Structure used to cache various metrics for our capture */
+struct FCachedMetrics
+{
+	FCachedMetrics() : Width(0), Height(0), Frame(0), ElapsedSeconds(0.f) {}
+
+	/** The width/Height of the frame */
+	int32 Width, Height;
+	/** The current frame number */
+	int32 Frame;
+	/** The number of seconds that have elapsed */
+	float ElapsedSeconds;
+};
+
 /** Class responsible for capturing scene data */
 UCLASS(config=EditorSettings)
 class MOVIESCENECAPTURE_API UMovieSceneCapture : public UObject, public IMovieSceneCaptureInterface
@@ -58,6 +71,9 @@ public:
 
 public:
 
+	/** Access this object's cached metrics */
+	const FCachedMetrics& GetMetrics() const { return CachedMetrics; }
+
 	/** Initialize the capture */
 	void StartCapture();
 
@@ -86,17 +102,6 @@ protected:
 	FString ResolveUniqueFilename();
 
 protected:
-
-	/** Structure used to cache various metrics for our capture */
-	struct FCachedMetrics
-	{
-		FCachedMetrics() : Width(0), Height(0), Frame(0) {}
-
-		/** The width/Height of the frame */
-		int32 Width, Height;
-		/** The current frame number */
-		int32 Frame;
-	};
 	
 	/** The viewport we are bound to */
 	FViewport* Viewport;
@@ -104,8 +109,6 @@ protected:
 	FMovieSceneCaptureHandle Handle;
 	/** Cached metrics for this capture operation */
 	FCachedMetrics CachedMetrics;
-	/** The current capture time, in seconds */
-	double CurrentTimeSeconds;
 	/** Optional AVI file writer used when capturing to video */
 	TUniquePtr<FAVIWriter> AVIWriter;
 	/** Strategy used for capture (real-time/fixed-time-step) */
diff --git a/Engine/Source/Runtime/MovieSceneCapture/Public/MovieSceneCaptureEnvironment.h b/Engine/Source/Runtime/MovieSceneCapture/Public/MovieSceneCaptureEnvironment.h
new file mode 100644
--- /dev/null
+++ b/Engine/Source/Runtime/MovieSceneCapture/Public/MovieSceneCaptureEnvironment.h
@@ -0,0 +1,20 @@
+// Copyright 1998-2015 Epic Games, Inc. All Rights Reserved.
+
+#pragma once
+
+#include "MovieSceneCaptureEnvironment.generated.h"
+
+UCLASS()
+class UMovieSceneCaptureEnvironment : public UObject
+{
+public:
+	GENERATED_BODY()
+
+	/** Get the frame number of the current capture */
+	UFUNCTION(BlueprintPure, Category="Cinematics|Capture")
+	static int32 GetCaptureFrameNumber();
+
+	/** Get the total elapsed time of the current capture in seconds */
+	UFUNCTION(BlueprintPure, Category="Cinematics|Capture")
+	static float GetCaptureElapsedTime();
+};
