Movie scene capture notification now expires automatically

diff --git a/Engine/Source/Runtime/MovieSceneCapture/Private/MovieSceneCaptureModule.cpp b/Engine/Source/Runtime/MovieSceneCapture/Private/MovieSceneCaptureModule.cpp
--- a/Engine/Source/Runtime/MovieSceneCapture/Private/MovieSceneCaptureModule.cpp
+++ b/Engine/Source/Runtime/MovieSceneCapture/Private/MovieSceneCaptureModule.cpp
@@ -126,13 +126,11 @@ class SCaptureMovieNotification : public SCompoundWidget, public INotificationWi
 public:
 	SLATE_BEGIN_ARGS(SCaptureMovieNotification){}
 		SLATE_EVENT(FOnProcessClosed, OnProcessClosed)
-		SLATE_EVENT(FSimpleDelegate, OnNotificationClosed)
 		SLATE_ARGUMENT(FString, BrowseToFolder)
 	SLATE_END_ARGS()
 
 	void Construct(const FArguments& InArgs, FProcHandle InProcHandle)
 	{
-		OnNotificationClosed = InArgs._OnNotificationClosed;
 		OnProcessClosed = InArgs._OnProcessClosed;
 		ProcHandle = InProcHandle;
 
@@ -168,37 +166,32 @@ public:
 					.AutoWidth()
 					.Padding(FMargin(15.f,0,0,0))
 					[
-						SNew(SThrobber)
-						.Visibility( this, &SCaptureMovieNotification::GetThrobberVisibility )
+						SAssignNew(Throbber, SThrobber)
 					]
 				]
 
 				+ SVerticalBox::Slot()
 				.AutoHeight()
+				.HAlign(HAlign_Right)
 				[
 					SNew(SHorizontalBox)
 
 					+ SHorizontalBox::Slot()
 					.AutoWidth()
 					.VAlign(VAlign_Center)
-					.Padding(FMargin(0,0,50.0f,0))
 					[
-						SNew(SHyperlink)
-						.Visibility( this, &SCaptureMovieNotification::GetHyperlinkVisibility )
+						SAssignNew(Hyperlink, SHyperlink)
+						.Visibility(EVisibility::Collapsed)
 						.Text(LOCTEXT("OpenFolder", "Open Capture Folder..."))
 						.OnNavigate_Lambda(OnBrowseToFolder)
 					]
 
-					+ SHorizontalBox::Slot()
-					[
-						SNew(SSpacer)
-					]
-
 					+ SHorizontalBox::Slot()
 					.AutoWidth()
 					.VAlign(VAlign_Center)
 					[
 						SAssignNew(Button, SButton)
+						.Text(LOCTEXT("StopButton", "Stop Capture"))
 						.OnClicked(this, &SCaptureMovieNotification::ButtonClicked)
 					]
 				]
@@ -229,13 +222,11 @@ public:
 	virtual void OnSetCompletionState(SNotificationItem::ECompletionState InState)
 	{
 		State = InState;
-		if (State == SNotificationItem::CS_Pending)
-		{
-			Button->SetContent(SNew(STextBlock).Text(LOCTEXT("StopButton", "Stop")));
-		}
-		else
+		if (State != SNotificationItem::CS_Pending)
 		{
-			Button->SetContent(SNew(STextBlock).Text(LOCTEXT("CloseButton", "Close")));
+			Hyperlink->SetVisibility(EVisibility::Visible);
+			Throbber->SetVisibility(EVisibility::Collapsed);
+			Button->SetVisibility(EVisibility::Collapsed);
 		}
 	}
 
@@ -246,35 +237,20 @@ public:
 
 private:
 
-	EVisibility GetThrobberVisibility() const
-	{
-		return State == SNotificationItem::CS_Pending ? EVisibility::Visible : EVisibility::Collapsed;
-	}
-
-	EVisibility GetHyperlinkVisibility() const
-	{
-		return State != SNotificationItem::CS_Pending ? EVisibility::Visible : EVisibility::Hidden;
-	}
-
 	FReply ButtonClicked()
 	{
 		if (State == SNotificationItem::CS_Pending)
 		{
 			FPlatformProcess::TerminateProc(ProcHandle);
 		}
-		else
-		{
-			OnNotificationClosed.ExecuteIfBound();
-		}
 		return FReply::Handled();
 	}
 	
 private:
-	TSharedPtr<SButton> Button;
+	TSharedPtr<SWidget> Button, Throbber, Hyperlink;
 	TSharedPtr<STextBlock> TextBlock;
 	SNotificationItem::ECompletionState State;
 	FOnProcessClosed OnProcessClosed;
-	FSimpleDelegate OnNotificationClosed;
 	FProcHandle ProcHandle;
 };
 
@@ -471,12 +447,6 @@ private:
 #endif
 	}
 
-	void OnNotificationClosed()
-	{
-		InProgressCaptureNotification->Fadeout();
-		InProgressCaptureNotification = nullptr;
-	}
-
 	void OnMovieCaptureProcessClosed(int32 RetCode)
 	{
 		if (RetCode == 0)
@@ -488,6 +458,9 @@ private:
 			// todo: error to message log
 			InProgressCaptureNotification->SetCompletionState(SNotificationItem::CS_Fail);
 		}
+
+		InProgressCaptureNotification->ExpireAndFadeout();
+		InProgressCaptureNotification = nullptr;
 	}
 
 	FText OnStartCapture(UMovieSceneCapture* CaptureObject)
@@ -588,9 +561,9 @@ private:
 				SNew(SCaptureMovieNotification, ProcessHandle)
 				.BrowseToFolder(CaptureObject->Settings.OutputDirectory.Path)
 				.OnProcessClosed_Raw(this, &FMovieSceneCaptureModule::OnMovieCaptureProcessClosed)
-				.OnNotificationClosed_Raw(this, &FMovieSceneCaptureModule::OnNotificationClosed)
 			);
 			Info.bFireAndForget = false;
+			Info.ExpireDuration = 5.f;
 			InProgressCaptureNotification = FSlateNotificationManager::Get().AddNotification(Info);
 			InProgressCaptureNotification->SetCompletionState(SNotificationItem::CS_Pending);
 		}
