Changed display name for movie scene capture settings filename format

diff --git a/Engine/Source/Runtime/MovieSceneCapture/Public/MovieSceneCaptureSettings.h b/Engine/Source/Runtime/MovieSceneCapture/Public/MovieSceneCaptureSettings.h
--- a/Engine/Source/Runtime/MovieSceneCapture/Public/MovieSceneCaptureSettings.h
+++ b/Engine/Source/Runtime/MovieSceneCapture/Public/MovieSceneCaptureSettings.h
@@ -46,7 +46,7 @@ struct MOVIESCENECAPTURE_API FMovieSceneCaptureSettings
 	 * {width}		- The width of the captured frames
 	 * {height}		- The height of the captured frames
 	 */
-	UPROPERTY(config, EditAnywhere, Category=General, AdvancedDisplay)
+	UPROPERTY(config, EditAnywhere, Category=General, AdvancedDisplay, DisplayName="Filename Format")
 	FString OutputFormat;
 
 	/** The frame rate at which to capture */
