Copying //UE4/Orion-Staging to Dev-Main (Originates from //Orion/Main @ CL-2759376)

diff --git a/Engine/Source/Runtime/MovieSceneCapture/Public/MovieSceneCaptureSettings.h b/Engine/Source/Runtime/MovieSceneCapture/Public/MovieSceneCaptureSettings.h
--- a/Engine/Source/Runtime/MovieSceneCapture/Public/MovieSceneCaptureSettings.h
+++ b/Engine/Source/Runtime/MovieSceneCapture/Public/MovieSceneCaptureSettings.h
@@ -41,42 +41,67 @@ struct MOVIESCENECAPTURE_API FMovieSceneCaptureSettings
 	UPROPERTY(config, EditAnywhere, Category=General, AdvancedDisplay, meta=(RelativePath))
 	FDirectoryPath OutputDirectory;
 
+	/** Whether to save temporary copies of all of the levels before capturing the movie.  This allows you to record movies of temporary work, or work that isn't yet saved, but it will take much longer for capturing to begin. */
+	UPROPERTY(config, EditAnywhere, Category=General, AdvancedDisplay)
+	bool bCreateTemporaryCopiesOfLevels;
+
+	/** Optional game mode to override the map's default game mode with.  This can be useful if the game's normal mode displays UI elements or loading screens that you don't want captured. */
+	UPROPERTY(config, EditAnywhere, Category=General, AdvancedDisplay)
+	TSubclassOf<class AGameMode> GameModeOverride;
+
 	/** The format to use for the resulting filename. Extension will be added automatically. Any tokens of the form {token} will be replaced with the corresponding value:
 	 * {fps}		- The captured framerate
 	 * {frame}		- The current frame number (only relevant for image sequences)
 	 * {width}		- The width of the captured frames
 	 * {height}		- The height of the captured frames
 	 * {world}		- The name of the current world
+	 * {quality}	- The image compression quality setting
 	 */
 	UPROPERTY(config, EditAnywhere, Category=General, DisplayName="Filename Format")
 	FString OutputFormat;
 
 	/** Whether to overwrite existing files or not */
-	UPROPERTY(config, EditAnywhere, Category=General)
+	UPROPERTY(config, EditAnywhere, Category=General, AdvancedDisplay )
 	bool bOverwriteExisting;
 
+	/** True if frame numbers in the output files should be relative to zero, rather than the actual frame numbers in the originating animation content */
+	UPROPERTY(config, EditAnywhere, Category=General, AdvancedDisplay)
+	bool bUseRelativeFrameNumbers;
+
 	/** The frame rate at which to capture */
 	UPROPERTY(config, EditAnywhere, Category=CaptureSettings)
 	int32 FrameRate;
 
-	/** An explicit number of frames to capture, or 0 if unspecified. Not exposed on UI - calculated internally */
-	UPROPERTY()
-	int32 FrameCount;
+	/** When enabled, the StartFrame setting will override the default starting frame number */
+	UPROPERTY(config, EditAnywhere, Category=CaptureSettings, AdvancedDisplay )
+	bool bUseCustomStartFrame;
+
+	/** Frame number to start capturing.  The frame number range depends on whether the bUseRelativeFrameNumbers option is enabled. */
+	UPROPERTY(config, EditAnywhere, Category=CaptureSettings, AdvancedDisplay, meta=(EditCondition="bUseCustomStartFrame") )
+	int32 StartFrame;
+
+	/** When enabled, the EndFrame setting will override the default ending frame number */
+	UPROPERTY(config, EditAnywhere, Category=CaptureSettings, AdvancedDisplay )
+	bool bUseCustomEndFrame;
+
+	/** Frame number to end capturing.  The frame number range depends on whether the bUseRelativeFrameNumbers option is enabled. */
+	UPROPERTY(config, EditAnywhere, Category=CaptureSettings, AdvancedDisplay, meta=(EditCondition="bUseCustomEndFrame") )
+	int32 EndFrame;
 
 	/** The resolution at which to capture */
 	UPROPERTY(config, EditAnywhere, Category=CaptureSettings, meta=(ShowOnlyInnerProperties))
 	FCaptureResolution Resolution;
 
 	/** The type of capture to perform */
-	UPROPERTY(config, EditAnywhere, Category=CaptureSettings, AdvancedDisplay)
+	UPROPERTY(config, EditAnywhere, Category=CaptureSettings)
 	EMovieCaptureType CaptureType;
 
 	/** Whether compression is enabled on the resulting file(s) */
-	UPROPERTY(config, EditAnywhere, Category=CaptureSettings)
+	UPROPERTY(config, EditAnywhere, Category=CaptureSettings, AdvancedDisplay )
 	bool bUseCompression;
 
-	/** The level of compression to apply to the captured file(s) (0-1) */
-	UPROPERTY(config, EditAnywhere, Category=CaptureSettings, AdvancedDisplay, meta=(ClampMin=0, ClampMax=1))
+	/** For output formats that support compression, this specifies the quality level of the compressed image.  0.0 results in the most compression, while 1.0 gives you the highest image quality. */
+	UPROPERTY(config, EditAnywhere, Category=CaptureSettings, AdvancedDisplay, meta=(ClampMin=0, ClampMax=1), meta=(EditCondition="bUseCompression"))
 	float CompressionQuality;
 
 	/** (Experimental) - An optional codec to use for video encoding */
@@ -92,18 +117,18 @@ struct MOVIESCENECAPTURE_API FMovieSceneCaptureSettings
 	bool bCinematicMode;
 
 	/** Whether to allow player movement whilst capturing */
-	UPROPERTY(config, EditAnywhere, Category=Cinematic, AdvancedDisplay)
+	UPROPERTY(config, EditAnywhere, Category=Cinematic, AdvancedDisplay, meta=(EditCondition="bCinematicMode"))
 	bool bAllowMovement;
 
 	/** Whether to allow player rotation whilst capturing */
-	UPROPERTY(config, EditAnywhere, Category=Cinematic, AdvancedDisplay)
+	UPROPERTY(config, EditAnywhere, Category=Cinematic, AdvancedDisplay, meta=(EditCondition="bCinematicMode"))
 	bool bAllowTurning;
 
 	/** Whether to show the local player whilst capturing */
-	UPROPERTY(config, EditAnywhere, Category=Cinematic, AdvancedDisplay)
+	UPROPERTY(config, EditAnywhere, Category=Cinematic, AdvancedDisplay, meta=(EditCondition="bCinematicMode"))
 	bool bShowPlayer;
 
 	/** Whether to show the in-game HUD whilst capturing */
-	UPROPERTY(config, EditAnywhere, Category=Cinematic, AdvancedDisplay)
+	UPROPERTY(config, EditAnywhere, Category=Cinematic, AdvancedDisplay, meta=(EditCondition="bCinematicMode"))
 	bool bShowHUD;
 };
\ No newline at end of file
diff --git a/Engine/Source/Runtime/MovieSceneTracks/Private/Sections/MovieSceneEventSection.cpp b/Engine/Source/Runtime/MovieSceneTracks/Private/Sections/MovieSceneEventSection.cpp
--- a/Engine/Source/Runtime/MovieSceneTracks/Private/Sections/MovieSceneEventSection.cpp
+++ b/Engine/Source/Runtime/MovieSceneTracks/Private/Sections/MovieSceneEventSection.cpp
@@ -31,7 +31,7 @@ void UMovieSceneEventSection::TriggerEvents(ALevelScriptActor* LevelScriptActor,
 	{
 		for (const auto& Key : Keys)
 		{
-			if ((Key.Time > LastPosition) && (Key.Time <= Position))
+			if ((Key.Time >= LastPosition) && (Key.Time <= Position))
 			{
 				TriggerEvent(Key.Value, LevelScriptActor);
 			}
@@ -43,7 +43,7 @@ void UMovieSceneEventSection::TriggerEvents(ALevelScriptActor* LevelScriptActor,
 		{
 			const auto& Key = Keys[KeyIndex];
 
-			if ((Key.Time >= Position) && (Key.Time < LastPosition))
+			if ((Key.Time >= Position) && (Key.Time <= LastPosition))
 			{
 				TriggerEvent(Key.Value, LevelScriptActor);
 			}
diff --git a/Engine/Source/Runtime/MovieSceneTracks/Private/Sections/MovieSceneSkeletalAnimationSection.cpp b/Engine/Source/Runtime/MovieSceneTracks/Private/Sections/MovieSceneSkeletalAnimationSection.cpp
--- a/Engine/Source/Runtime/MovieSceneTracks/Private/Sections/MovieSceneSkeletalAnimationSection.cpp
+++ b/Engine/Source/Runtime/MovieSceneTracks/Private/Sections/MovieSceneSkeletalAnimationSection.cpp
@@ -3,6 +3,7 @@
 #include "MovieSceneTracksPrivatePCH.h"
 #include "MovieSceneSkeletalAnimationSection.h"
 
+FName UMovieSceneSkeletalAnimationSection::DefaultSlotName( "DefaultSlot" );
 
 UMovieSceneSkeletalAnimationSection::UMovieSceneSkeletalAnimationSection( const FObjectInitializer& ObjectInitializer )
 	: Super( ObjectInitializer )
@@ -15,6 +16,7 @@ UMovieSceneSkeletalAnimationSection::UMovieSceneSkeletalAnimationSection( const
 	PreviousPlayRate = PlayRate;
 #endif
 	bReverse = false;
+	SlotName = DefaultSlotName;
 }
 
 
@@ -61,7 +63,10 @@ void UMovieSceneSkeletalAnimationSection::GetSnapTimes(TArray<float>& OutSnapTim
 	Super::GetSnapTimes(OutSnapTimes, bGetSectionBorders);
 
 	float CurrentTime = GetStartTime();
+	float AnimPlayRate = FMath::IsNearlyZero(GetPlayRate()) ? 1.0f : GetPlayRate();
+	float SeqLength = (GetSequenceLength() - (GetStartOffset() + GetEndOffset())) / AnimPlayRate;
 
+	// Snap to the repeat times
 	while (CurrentTime <= GetEndTime() && !FMath::IsNearlyZero(GetDuration()))
 	{
 		if (CurrentTime >= GetStartTime())
@@ -69,7 +74,7 @@ void UMovieSceneSkeletalAnimationSection::GetSnapTimes(TArray<float>& OutSnapTim
 			OutSnapTimes.Add(CurrentTime);
 		}
 
-		CurrentTime += GetDuration() - (StartOffset + EndOffset);
+		CurrentTime += SeqLength;
 	}
 }
 
diff --git a/Engine/Source/Runtime/MovieSceneTracks/Private/TrackInstances/MovieSceneSkeletalAnimationTrackInstance.cpp b/Engine/Source/Runtime/MovieSceneTracks/Private/TrackInstances/MovieSceneSkeletalAnimationTrackInstance.cpp
--- a/Engine/Source/Runtime/MovieSceneTracks/Private/TrackInstances/MovieSceneSkeletalAnimationTrackInstance.cpp
+++ b/Engine/Source/Runtime/MovieSceneTracks/Private/TrackInstances/MovieSceneSkeletalAnimationTrackInstance.cpp
@@ -59,7 +59,6 @@ void FMovieSceneSkeletalAnimationTrackInstance::Update( float Position, float La
 				Position -= 1 / 1000.0f;
 
 				int32 ChannelIndex = 0;
-				FName SlotName = FName("AnimationSlot");
 				UAnimSequence* AnimSequence = AnimSection->GetAnimSequence();
 
 				float AnimPlayRate = FMath::IsNearlyZero(AnimSection->GetPlayRate()) ? 1.0f : AnimSection->GetPlayRate();
@@ -80,12 +79,12 @@ void FMovieSceneSkeletalAnimationTrackInstance::Update( float Position, float La
 				if (GIsEditor && !GWorld->HasBegunPlay())
 				{
 					AnimInterface->PreviewBeginAnimControl(nullptr);
-					AnimInterface->PreviewSetAnimPosition(SlotName, ChannelIndex, AnimSequence, AnimPosition, true, false, 0.f);
+					AnimInterface->PreviewSetAnimPosition(AnimSection->GetSlotName(), ChannelIndex, AnimSequence, AnimPosition, true, false, 0.f);
 				}
 				else
 				{
 					AnimInterface->BeginAnimControl(nullptr);
-					AnimInterface->SetAnimPosition(SlotName, ChannelIndex, AnimSequence, AnimPosition, true, false);
+					AnimInterface->SetAnimPosition(AnimSection->GetSlotName(), ChannelIndex, AnimSequence, AnimPosition, true, false);
 				}
 			}
 		}
diff --git a/Engine/Source/Runtime/MovieSceneTracks/Private/Tracks/MovieSceneParticleTrack.cpp b/Engine/Source/Runtime/MovieSceneTracks/Private/Tracks/MovieSceneParticleTrack.cpp
--- a/Engine/Source/Runtime/MovieSceneTracks/Private/Tracks/MovieSceneParticleTrack.cpp
+++ b/Engine/Source/Runtime/MovieSceneTracks/Private/Tracks/MovieSceneParticleTrack.cpp
@@ -72,7 +72,7 @@ void UMovieSceneParticleTrack::AddNewSection( float SectionTime )
 {
 	if ( MovieSceneHelpers::FindSectionAtTime( ParticleSections, SectionTime ) == nullptr )
 	{
-		UMovieSceneParticleSection* NewSection = NewObject<UMovieSceneParticleSection>( this );
+		UMovieSceneParticleSection* NewSection = Cast<UMovieSceneParticleSection>( CreateNewSection() );
 		NewSection->SetStartTime( SectionTime );
 		NewSection->SetEndTime( SectionTime );
 		NewSection->SetStartTime( SectionTime );
@@ -81,5 +81,10 @@ void UMovieSceneParticleTrack::AddNewSection( float SectionTime )
 	}
 }
 
+UMovieSceneSection* UMovieSceneParticleTrack::CreateNewSection()
+{
+	return NewObject<UMovieSceneParticleSection>( this );
+}
+
 
 #undef LOCTEXT_NAMESPACE
diff --git a/Engine/Source/Runtime/MovieSceneTracks/Private/Tracks/MovieSceneSkeletalAnimationTrack.cpp b/Engine/Source/Runtime/MovieSceneTracks/Private/Tracks/MovieSceneSkeletalAnimationTrack.cpp
--- a/Engine/Source/Runtime/MovieSceneTracks/Private/Tracks/MovieSceneSkeletalAnimationTrack.cpp
+++ b/Engine/Source/Runtime/MovieSceneTracks/Private/Tracks/MovieSceneSkeletalAnimationTrack.cpp
@@ -23,7 +23,7 @@ UMovieSceneSkeletalAnimationTrack::UMovieSceneSkeletalAnimationTrack( const FObj
 
 void UMovieSceneSkeletalAnimationTrack::AddNewAnimation(float KeyTime, UAnimSequence* AnimSequence)
 {
-	UMovieSceneSkeletalAnimationSection* NewSection = NewObject<UMovieSceneSkeletalAnimationSection>(this);
+	UMovieSceneSkeletalAnimationSection* NewSection = Cast<UMovieSceneSkeletalAnimationSection>(CreateNewSection());
 	{
 		NewSection->InitialPlacement(AnimationSections, KeyTime, KeyTime + AnimSequence->SequenceLength, SupportsMultipleRows());
 		NewSection->SetAnimSequence(AnimSequence);
@@ -43,8 +43,7 @@ UMovieSceneSection* UMovieSceneSkeletalAnimationTrack::GetAnimSectionAtTime(floa
 		}
 	}
 
-	UMovieSceneSection* NearestSection = MovieSceneHelpers::FindNearestSectionAtTime( AnimationSections, Time );
-	return NearestSection;
+	return nullptr;
 }
 
 
@@ -63,9 +62,15 @@ const TArray<UMovieSceneSection*>& UMovieSceneSkeletalAnimationTrack::GetAllSect
 }
 
 
+UMovieSceneSection* UMovieSceneSkeletalAnimationTrack::CreateNewSection()
+{
+	return NewObject<UMovieSceneSkeletalAnimationSection>( this );
+}
+
+
 void UMovieSceneSkeletalAnimationTrack::RemoveAllAnimationData()
 {
-	// do nothing
+	AnimationSections.Empty();
 }
 
 
diff --git a/Engine/Source/Runtime/MovieSceneTracks/Public/Sections/MovieSceneSkeletalAnimationSection.h b/Engine/Source/Runtime/MovieSceneTracks/Public/Sections/MovieSceneSkeletalAnimationSection.h
--- a/Engine/Source/Runtime/MovieSceneTracks/Public/Sections/MovieSceneSkeletalAnimationSection.h
+++ b/Engine/Source/Runtime/MovieSceneTracks/Public/Sections/MovieSceneSkeletalAnimationSection.h
@@ -49,6 +49,12 @@ public:
 	/** Sets whether the playback is reversed */
 	void SetReverse(bool bInReverse) { bReverse = bInReverse; }
 
+	/** Gets the anim BP slot name. */
+	FName GetSlotName() const { return SlotName; }
+
+	/** Sets the anim BP slot name. */
+	void SetSlotName( FName InSlotName ) { SlotName = InSlotName; }
+
 	/** MovieSceneSection interface */
 	virtual void MoveSection( float DeltaPosition, TSet<FKeyHandle>& KeyHandles ) override;
 	virtual void DilateSection( float DilationFactor, float Origin, TSet<FKeyHandle>& KeyHandles  ) override;
@@ -58,6 +64,8 @@ public:
 
 private:
 
+	static FName DefaultSlotName;
+
 	// UObject interface
 #if WITH_EDITOR
 	virtual void PreEditChange(UProperty* PropertyAboutToChange) override;
@@ -84,4 +92,8 @@ private:
 	/** Reverse the playback of the animation clip */
 	UPROPERTY(EditAnywhere, Category="Animation")
 	uint32 bReverse:1;
+
+	/** The slot name to use for the animation */
+	UPROPERTY( EditAnywhere, Category = "Animation" )
+	FName SlotName;
 };
diff --git a/Engine/Source/Runtime/MovieSceneTracks/Public/Tracks/MovieSceneParticleTrack.h b/Engine/Source/Runtime/MovieSceneTracks/Public/Tracks/MovieSceneParticleTrack.h
--- a/Engine/Source/Runtime/MovieSceneTracks/Public/Tracks/MovieSceneParticleTrack.h
+++ b/Engine/Source/Runtime/MovieSceneTracks/Public/Tracks/MovieSceneParticleTrack.h
@@ -43,6 +43,7 @@ public:
 	virtual TRange<float> GetSectionBoundaries() const override;
 	virtual const TArray<UMovieSceneSection*>& GetAllSections() const override;
 	virtual void AddNewSection(float SectionTime);
+	virtual UMovieSceneSection* CreateNewSection() override;
 
 private:
 
diff --git a/Engine/Source/Runtime/MovieSceneTracks/Public/Tracks/MovieSceneSkeletalAnimationTrack.h b/Engine/Source/Runtime/MovieSceneTracks/Public/Tracks/MovieSceneSkeletalAnimationTrack.h
--- a/Engine/Source/Runtime/MovieSceneTracks/Public/Tracks/MovieSceneSkeletalAnimationTrack.h
+++ b/Engine/Source/Runtime/MovieSceneTracks/Public/Tracks/MovieSceneSkeletalAnimationTrack.h
@@ -40,6 +40,7 @@ public:
 	virtual TRange<float> GetSectionBoundaries() const override;
 	virtual bool SupportsMultipleRows() const override { return true; }
 	virtual const TArray<UMovieSceneSection*>& GetAllSections() const override;
+	virtual UMovieSceneSection* CreateNewSection() override;
 
 #if WITH_EDITORONLY_DATA
 	virtual FText GetDisplayName() const override;
diff --git a/Engine/Source/Runtime/Online/HTTP/Private/HttpRetrySystem.cpp b/Engine/Source/Runtime/Online/HTTP/Private/HttpRetrySystem.cpp
--- a/Engine/Source/Runtime/Online/HTTP/Private/HttpRetrySystem.cpp
+++ b/Engine/Source/Runtime/Online/HTTP/Private/HttpRetrySystem.cpp
@@ -276,7 +276,7 @@ bool FHttpRetrySystem::FManager::Update(uint32* FileCount, uint32* FailingCount,
         }
 		else
 		{
-			UE_LOG(LogHttp, Warning, TEXT("Timeout on %s"), HttpRetryRequestEntry.CurrentRetryCount + 1, *(HttpRetryRequest->GetURL()));
+			UE_LOG(LogHttp, Warning, TEXT("Timeout on retry %d: %s"), HttpRetryRequestEntry.CurrentRetryCount + 1, *(HttpRetryRequest->GetURL()));
 			bIsGreen = false;
             HttpRetryRequest->Status = FHttpRetrySystem::FRequest::EStatus::FailedTimeout;
 			if (FailedCount != nullptr)
@@ -326,7 +326,7 @@ bool FHttpRetrySystem::FManager::Update(uint32* FileCount, uint32* FailingCount,
 FHttpRetrySystem::FManager::FHttpRetryRequestEntry::FHttpRetryRequestEntry(TSharedRef<FHttpRetrySystem::FRequest>& InHttpRequest)
     : bShouldCancel(false)
     , CurrentRetryCount(0)
-    , RequestStartTimeAbsoluteSeconds(0.0)
+	, RequestStartTimeAbsoluteSeconds(FPlatformTime::Seconds())
     , HttpRequest(InHttpRequest)
 {}
 
diff --git a/Engine/Source/Runtime/Online/OnlineFramework/Party/Private/PartyGameState.cpp b/Engine/Source/Runtime/Online/OnlineFramework/Party/Private/PartyGameState.cpp
--- a/Engine/Source/Runtime/Online/OnlineFramework/Party/Private/PartyGameState.cpp
+++ b/Engine/Source/Runtime/Online/OnlineFramework/Party/Private/PartyGameState.cpp
@@ -300,11 +300,6 @@ UPartyMemberState* UPartyGameState::CreateNewPartyMember(const FUniqueNetId& InM
 				NewPartyMemberState = NewObject<UPartyMemberState>(this, PartyMemberStateClass);
 				NewPartyMemberState->UniqueId.SetUniqueNetId(PartyMember->GetUserId());
 				NewPartyMemberState->DisplayName = FText::FromString(PartyMember->GetDisplayName());
-
-				if ( InitMemberFunctor )
-				{
-					InitMemberFunctor( PartyMember, NewPartyMemberState );
-				}
 			}
 			else
 			{
@@ -324,11 +319,6 @@ UPartyMemberState* UPartyGameState::CreateNewPartyMember(const FUniqueNetId& InM
 	return NewPartyMemberState;
 }
 
-void UPartyGameState::SetInitMemberFunctor( TFunction<void( TSharedPtr<FOnlinePartyMember>, UPartyMemberState* )>&& InFunctor )
-{
-	InitMemberFunctor = MoveTemp( InFunctor );
-}
-
 void UPartyGameState::HandlePartyConfigChanged( const TSharedRef<FPartyConfiguration>& InPartyConfig )
 {
 	UE_LOG(LogParty, VeryVerbose, TEXT("[%s] HandlePartyConfigChanged"), OssParty.IsValid() ? *OssParty->PartyId->ToString() : TEXT("INVALID"));
diff --git a/Engine/Source/Runtime/Online/OnlineFramework/Party/Public/PartyGameState.h b/Engine/Source/Runtime/Online/OnlineFramework/Party/Public/PartyGameState.h
--- a/Engine/Source/Runtime/Online/OnlineFramework/Party/Public/PartyGameState.h
+++ b/Engine/Source/Runtime/Online/OnlineFramework/Party/Public/PartyGameState.h
@@ -164,14 +164,6 @@ public:
 	static void AddReferencedObjects(UObject* InThis, FReferenceCollector& Collector);
 	// UObject interface end
 
-	
-	/** 
-	 * Allows external systems an opportunity to initialize newly created PartyMemberStates before being entered into the party system
-	 * 
-	 * @param InFunctor a TFunction to call immediately after allocation. 
-	 */
-	void SetInitMemberFunctor( TFunction<void( TSharedPtr<FOnlinePartyMember>, UPartyMemberState* )>&& InFunctor );
-
 	/**
 	 * Notification that the game is about to travel to another map/server
 	 */
@@ -616,8 +608,6 @@ private:
 	/** Scratch copy of child USTRUCT for handling replication comparisons */
 	FPartyState* PartyStateRefScratch;
 
-	TFunction<void( TSharedPtr<FOnlinePartyMember>, UPartyMemberState* )> InitMemberFunctor;
-
 	friend UParty;
 	friend UPartyMemberState;
 };
diff --git a/Engine/Source/Runtime/Online/OnlineSubsystem/Public/Interfaces/OnlinePresenceInterface.h b/Engine/Source/Runtime/Online/OnlineSubsystem/Public/Interfaces/OnlinePresenceInterface.h
--- a/Engine/Source/Runtime/Online/OnlineSubsystem/Public/Interfaces/OnlinePresenceInterface.h
+++ b/Engine/Source/Runtime/Online/OnlineSubsystem/Public/Interfaces/OnlinePresenceInterface.h
@@ -19,6 +19,9 @@ const FString CustomPresenceDataKey = "CustomData";
 /** Name of the client that sent the presence update */
 const FString DefaultAppIdKey = "AppId";
 
+/** Name of the platform the the presence update */
+const FString DefaultPlatformKey = "Platform";
+
 /** Override Id of the client to set the presence state to */
 const FString OverrideClientIdKey = "OverrideClientId";
 
diff --git a/Engine/Source/Runtime/Online/OnlineSubsystemUtils/Private/OnlinePIESettings.cpp b/Engine/Source/Runtime/Online/OnlineSubsystemUtils/Private/OnlinePIESettings.cpp
--- a/Engine/Source/Runtime/Online/OnlineSubsystemUtils/Private/OnlinePIESettings.cpp
+++ b/Engine/Source/Runtime/Online/OnlineSubsystemUtils/Private/OnlinePIESettings.cpp
@@ -94,22 +94,6 @@ void UOnlinePIESettings::PostInitProperties()
 {
 	Super::PostInitProperties();
 
-	// MIGRATION : Encrypt the plain text passwords if they are there
-	bool bNeedSave = false;
-	for (FPIELoginSettingsInternal& Login : Logins)
-	{
-		if (Login.Token.Len() > 0 && Login.TokenBytes.Num() == 0)
-		{
-			bNeedSave = true;
-			Login.Encrypt();
-		}
-	}
-	if (bNeedSave)
-	{
-		SaveConfig();
-	}
-    // MIGRATION END
-
 	for (FPIELoginSettingsInternal& Login : Logins)
 	{
 		Login.Decrypt();
diff --git a/Engine/Source/Runtime/Online/OnlineSubsystemUtils/Private/OnlinePIESettings.h b/Engine/Source/Runtime/Online/OnlineSubsystemUtils/Private/OnlinePIESettings.h
--- a/Engine/Source/Runtime/Online/OnlineSubsystemUtils/Private/OnlinePIESettings.h
+++ b/Engine/Source/Runtime/Online/OnlineSubsystemUtils/Private/OnlinePIESettings.h
@@ -17,8 +17,7 @@ public:
 	UPROPERTY(EditAnywhere, Category = "Logins", meta = (DisplayName = "User Id", Tooltip = "Id of the user logging in (email, display name, facebook id, etc)"))
 	FString Id;
 	/** Credentials of the user logging in (password or auth token) */
-	// MIGRATION: need to add Transient here after a bit now that we can encrypt passwords
-	UPROPERTY(EditAnywhere, Category = "Logins", meta = (DisplayName = "Password", Tooltip = "Credentials of the user logging in (password or auth token)", PasswordField = true))
+	UPROPERTY(EditAnywhere, Transient, Category = "Logins", meta = (DisplayName = "Password", Tooltip = "Credentials of the user logging in (password or auth token)", PasswordField = true))
 	FString Token;
 	/** Type of account. Needed to identity the auth method to use (epic, internal, facebook, etc) */
 	UPROPERTY(EditAnywhere, Category = "Logins", meta = (DisplayName = "Type", Tooltip = "Type of account. Needed to identity the auth method to use (epic, internal, facebook, etc)"))
diff --git a/Engine/Source/Runtime/RenderCore/Public/UniformBuffer.h b/Engine/Source/Runtime/RenderCore/Public/UniformBuffer.h
--- a/Engine/Source/Runtime/RenderCore/Public/UniformBuffer.h
+++ b/Engine/Source/Runtime/RenderCore/Public/UniformBuffer.h
@@ -30,6 +30,8 @@ public:
 		: BufferUsage(UniformBuffer_MultiFrame)
 	{
 		Contents = (uint8*)FMemory::Malloc(sizeof(TBufferStruct),UNIFORM_BUFFER_STRUCT_ALIGNMENT);
+		FMemory::Memzero(Contents, sizeof(TBufferStruct));
+		bEverUpdated = false;
 	}
 
 	~TUniformBuffer()
@@ -40,8 +42,7 @@ public:
 	/** Sets the contents of the uniform buffer. */
 	void SetContents(const TBufferStruct& NewContents)
 	{
-		check(IsInRenderingThread());
-		FMemory::Memcpy(Contents,&NewContents,sizeof(TBufferStruct));
+		SetContentsNoUpdate(NewContents);
 		UpdateRHI();
 	}
 	
@@ -57,13 +58,29 @@ public:
 	}
 
 	// Accessors.
-	FUniformBufferRHIParamRef GetUniformBufferRHI() const { return UniformBufferRHI; }
+	FUniformBufferRHIParamRef GetUniformBufferRHI() const 
+	{ 
+		ensureMsgf(bEverUpdated, TEXT("Caching UB before it was updated with valid data!"));
+		return UniformBufferRHI;
+	}
 
 	EUniformBufferUsage BufferUsage;
+
+protected:
+
+	/** Sets the contents of the uniform buffer. Used within calls to InitDynamicRHI */
+	void SetContentsNoUpdate(const TBufferStruct& NewContents)
+	{
+		check(IsInRenderingThread());
+		FMemory::Memcpy(Contents, &NewContents, sizeof(TBufferStruct));
+		bEverUpdated = true;
+	}
+
 private:
 
 	FUniformBufferRHIRef UniformBufferRHI;
 	uint8* Contents;
+	bool bEverUpdated;
 };
 
 /** A reference to a uniform buffer RHI resource with a specific structure. */
diff --git a/Engine/Source/Runtime/Renderer/Private/BasePassRendering.h b/Engine/Source/Runtime/Renderer/Private/BasePassRendering.h
--- a/Engine/Source/Runtime/Renderer/Private/BasePassRendering.h
+++ b/Engine/Source/Runtime/Renderer/Private/BasePassRendering.h
@@ -434,6 +434,8 @@ public:
 			const int32 VelocityIndex = (CVar && CVar->GetValueOnAnyThread() != 0) ? 6 : 5;
 			OutEnvironment.SetRenderTargetOutputFormat(VelocityIndex, PF_G16R16);
 		}
+
+		OutEnvironment.SetDefine(TEXT("TRANSLUCENCY_RENDERING"), 1);
 	}
 
 	/** Initialization constructor. */
@@ -444,6 +446,7 @@ public:
 		TranslucentLightingParameters.Bind(Initializer.ParameterMap);
 		EditorCompositeParams.Bind(Initializer.ParameterMap);
 		LightGrid.Bind(Initializer.ParameterMap,TEXT("LightGrid"));
+		ScreenTextureUVScale.Bind(Initializer.ParameterMap, TEXT("ScreenPositionUVScale"));
 	}
 	TBasePassPixelShaderBaseType() {}
 
@@ -454,7 +457,8 @@ public:
 		const FViewInfo* View, 
 		EBlendMode BlendMode, 
 		bool bEnableEditorPrimitveDepthTest,
-		ESceneRenderTargetsMode::Type TextureMode)
+		ESceneRenderTargetsMode::Type TextureMode,
+		float ScreenTextureScaleFactor = 1.0f)
 	{
 		const FPixelShaderRHIParamRef ShaderRHI = GetPixelShader();
 
@@ -464,6 +468,7 @@ public:
 		{
 			if (IsTranslucentBlendMode(BlendMode))
 			{
+				SetShaderValue(RHICmdList, ShaderRHI, ScreenTextureUVScale, FVector(ScreenTextureScaleFactor));
 				TranslucentLightingParameters.Set(RHICmdList, this, View);
 
 				// Experimental dynamic forward lighting for translucency. Can be the base for opaque forward lighting which will allow more lighting models or rendering without a GBuffer
@@ -501,6 +506,7 @@ public:
 		Ar << TranslucentLightingParameters;
  		Ar << EditorCompositeParams;
 		Ar << LightGrid;
+		Ar << ScreenTextureUVScale;
 		return bShaderHasOutdatedParameters;
 	}
 
@@ -508,6 +514,7 @@ private:
 	FTranslucentLightingParameters TranslucentLightingParameters;
 	FEditorCompositingParameters EditorCompositeParams;
 	FShaderResourceParameter LightGrid;
+	FShaderParameter ScreenTextureUVScale;
 };
 
 /** The concrete base pass pixel shader type. */
@@ -651,7 +658,7 @@ public:
 			LightMapPolicy == Other.LightMapPolicy;
 	}
 
-	void SetSharedState(FRHICommandList& RHICmdList, const FViewInfo* View, const ContextDataType PolicyContext) const
+	void SetSharedState(FRHICommandList& RHICmdList, const FViewInfo* View, const ContextDataType PolicyContext, float ScreenTextureScaleFactor = 1.0f) const
 	{
 		// Set the light-map policy.
 		LightMapPolicy.Set(RHICmdList, VertexShader,bOverrideWithShaderComplexity ? NULL : PixelShader,VertexShader,PixelShader,VertexFactory,MaterialRenderProxy,View);
@@ -685,7 +692,7 @@ public:
 		else
 #endif
 		{
-			PixelShader->SetParameters(RHICmdList, MaterialRenderProxy,*MaterialResource,View,BlendMode,bEnableEditorPrimitiveDepthTest,SceneTextureMode);
+			PixelShader->SetParameters(RHICmdList, MaterialRenderProxy, *MaterialResource, View, BlendMode, bEnableEditorPrimitiveDepthTest, SceneTextureMode, ScreenTextureScaleFactor);
 
 			switch(BlendMode)
 			{
diff --git a/Engine/Source/Runtime/Renderer/Private/PostProcess/SceneRenderTargets.cpp b/Engine/Source/Runtime/Renderer/Private/PostProcess/SceneRenderTargets.cpp
--- a/Engine/Source/Runtime/Renderer/Private/PostProcess/SceneRenderTargets.cpp
+++ b/Engine/Source/Runtime/Renderer/Private/PostProcess/SceneRenderTargets.cpp
@@ -1691,6 +1691,8 @@ void FSceneRenderTargets::AllocateDeferredShadingPathRenderTargets(FRHICommandLi
 						FClearValueBinding::None,
 						0,
 						TexCreate_ShaderResource | TexCreate_RenderTargetable,
+						false,
+						1,
 						false)),
 					TranslucencyLightingVolumeAmbient[RTSetIndex],
 					GetVolumeName(RTSetIndex, false)
@@ -1706,6 +1708,8 @@ void FSceneRenderTargets::AllocateDeferredShadingPathRenderTargets(FRHICommandLi
 						FClearValueBinding::None,
 						0,
 						TexCreate_ShaderResource | TexCreate_RenderTargetable,
+						false,
+						1,
 						false)),
 					TranslucencyLightingVolumeDirectional[RTSetIndex],
 					GetVolumeName(RTSetIndex, true)
diff --git a/Engine/Source/Runtime/Renderer/Private/Renderer.cpp b/Engine/Source/Runtime/Renderer/Private/Renderer.cpp
--- a/Engine/Source/Runtime/Renderer/Private/Renderer.cpp
+++ b/Engine/Source/Runtime/Renderer/Private/Renderer.cpp
@@ -75,11 +75,11 @@ void FRendererModule::DrawTileMesh(FRHICommandListImmediate& RHICmdList, const F
 			{
 				if (FeatureLevel >= ERHIFeatureLevel::SM4)
 				{
-					FBasePassOpaqueDrawingPolicyFactory::DrawDynamicMesh(RHICmdList, View, FBasePassOpaqueDrawingPolicyFactory::ContextType(false, ESceneRenderTargetsMode::SetTextures), Mesh, false, false, NULL, HitProxyId);
+					FBasePassOpaqueDrawingPolicyFactory::DrawDynamicMesh(RHICmdList, View, FBasePassOpaqueDrawingPolicyFactory::ContextType(false, ESceneRenderTargetsMode::DontSet), Mesh, false, false, NULL, HitProxyId);
 				}
 				else
 				{
-					FBasePassForwardOpaqueDrawingPolicyFactory::DrawDynamicMesh(RHICmdList, View, FBasePassForwardOpaqueDrawingPolicyFactory::ContextType(false, ESceneRenderTargetsMode::SetTextures), Mesh, false, false, NULL, HitProxyId);
+					FBasePassForwardOpaqueDrawingPolicyFactory::DrawDynamicMesh(RHICmdList, View, FBasePassForwardOpaqueDrawingPolicyFactory::ContextType(false, ESceneRenderTargetsMode::DontSet), Mesh, false, false, NULL, HitProxyId);
 				}
 			}
 		}	
diff --git a/Engine/Source/Runtime/Renderer/Private/TranslucentLighting.cpp b/Engine/Source/Runtime/Renderer/Private/TranslucentLighting.cpp
--- a/Engine/Source/Runtime/Renderer/Private/TranslucentLighting.cpp
+++ b/Engine/Source/Runtime/Renderer/Private/TranslucentLighting.cpp
@@ -1074,7 +1074,7 @@ void RasterizeToVolumeTexture(FRHICommandList& RHICmdList, FVolumeBounds VolumeB
 template<int32 NumRenderTargets>
 void ClearVolumeTextures(FRHICommandList& RHICmdList, ERHIFeatureLevel::Type FeatureLevel, const FTextureRHIParamRef* RenderTargets, const FLinearColor* ClearColors)
 {
-	SetRenderTargets(RHICmdList, NumRenderTargets, RenderTargets, FTextureRHIRef(), 0, NULL);
+	SetRenderTargets(RHICmdList, NumRenderTargets, RenderTargets, FTextureRHIRef(), 0, NULL, true);
 #if PLATFORM_XBOXONE
 	// ClearMRT is faster on Xbox
 	if (true)
@@ -1110,6 +1110,7 @@ void ClearVolumeTextures(FRHICommandList& RHICmdList, ERHIFeatureLevel::Type Fea
 
 		RasterizeToVolumeTexture(RHICmdList, VolumeBounds);
 	}
+	RHICmdList.TransitionResources(EResourceTransitionAccess::EReadable, (FTextureRHIParamRef*)RenderTargets, NumRenderTargets);
 }
 
 void FDeferredShadingSceneRenderer::ClearTranslucentVolumeLighting(FRHICommandListImmediate& RHICmdList)
@@ -1850,7 +1851,7 @@ void FDeferredShadingSceneRenderer::FilterTranslucentVolumeLighting(FRHICommandL
 				}
 				RHICmdList.TransitionResources(EResourceTransitionAccess::EReadable, Inputs, 2);
 
-				SetRenderTargets(RHICmdList, ARRAY_COUNT(RenderTargets), RenderTargets, FTextureRHIRef(), 0, NULL);
+				SetRenderTargets(RHICmdList, ARRAY_COUNT(RenderTargets), RenderTargets, FTextureRHIRef(), 0, NULL, true);
 
 
 				const FVolumeBounds VolumeBounds(GTranslucencyLightingVolumeDim);
diff --git a/Engine/Source/Runtime/Renderer/Private/TranslucentRendering.cpp b/Engine/Source/Runtime/Renderer/Private/TranslucentRendering.cpp
--- a/Engine/Source/Runtime/Renderer/Private/TranslucentRendering.cpp
+++ b/Engine/Source/Runtime/Renderer/Private/TranslucentRendering.cpp
@@ -11,6 +11,8 @@
 #include "LightPropagationVolume.h"
 #include "SceneUtils.h"
 
+static float GScreenTextureUVScale;
+
 static void SetTranslucentRenderTargetAndState(FRHICommandList& RHICmdList, const FViewInfo& View, ETranslucencyPassType TranslucenyPassType, bool bFirstTimeThisFrame = false)
 {
 	FSceneRenderTargets& SceneContext = FSceneRenderTargets::Get(RHICmdList);
@@ -227,6 +229,7 @@ public:
 	bool bBackFace;
 	float DitheredLODTransitionValue;
 	bool bUseTranslucentSelfShadowing;
+	float SeparateTranslucencyScreenTextureScaleFactor;
 
 	/** Initialization constructor. */
 	FDrawTranslucentMeshAction(
@@ -235,14 +238,16 @@ public:
 		float InDitheredLODTransitionValue,
 		FHitProxyId InHitProxyId,
 		const FProjectedShadowInfo* InTranslucentSelfShadow,
-		bool bInUseTranslucentSelfShadowing
-		):
+		bool bInUseTranslucentSelfShadowing,
+		float ScreenTextureUVScaleFactor
+		) :
 		View(InView),
 		TranslucentSelfShadow(InTranslucentSelfShadow),
 		HitProxyId(InHitProxyId),
 		bBackFace(bInBackFace),
 		DitheredLODTransitionValue(InDitheredLODTransitionValue),
-		bUseTranslucentSelfShadowing(bInUseTranslucentSelfShadowing)
+		bUseTranslucentSelfShadowing(bInUseTranslucentSelfShadowing),
+		SeparateTranslucencyScreenTextureScaleFactor(ScreenTextureUVScaleFactor)
 	{}
 
 	bool UseTranslucentSelfShadowing() const 
@@ -295,7 +300,7 @@ public:
 			Parameters.bAllowFog
 			);
 		RHICmdList.BuildAndSetLocalBoundShaderState(DrawingPolicy.GetBoundShaderStateInput(View.GetFeatureLevel()));
-		DrawingPolicy.SetSharedState(RHICmdList, &View, typename TBasePassDrawingPolicy<LightMapPolicyType>::ContextDataType());
+		DrawingPolicy.SetSharedState(RHICmdList, &View, typename TBasePassDrawingPolicy<LightMapPolicyType>::ContextDataType(), SeparateTranslucencyScreenTextureScaleFactor);
 
 		int32 BatchElementIndex = 0;
 		uint64 BatchElementMask = Parameters.BatchElementMask;
@@ -392,9 +397,9 @@ bool FTranslucencyDrawingPolicyFactory::DrawMesh(
 	// Only render translucent materials
 	if(IsTranslucentBlendMode(BlendMode))
 	{
+		bool bCurrentlyRenderingSeparateTranslucency = Material->IsSeparateTranslucencyEnabled() == (DrawingContext.TranslucenyPassType == TPT_SeparateTransluceny);
 		// if we are in relevant pass
-		if(Material->IsSeparateTranslucencyEnabled() == (DrawingContext.TranslucenyPassType == TPT_SeparateTransluceny)
-			|| bSeparateTranslucencyEnabled == false)
+		if (bCurrentlyRenderingSeparateTranslucency || bSeparateTranslucencyEnabled == false)
 		{
 			if (Material->RequiresSceneColorCopy_RenderThread())
 			{
@@ -457,7 +462,8 @@ bool FTranslucencyDrawingPolicyFactory::DrawMesh(
 					DitheredLODTransitionValue,
 					HitProxyId,
 					DrawingContext.TranslucentSelfShadow,
-					PrimitiveSceneProxy && PrimitiveSceneProxy->CastsVolumetricTranslucentShadow()
+					PrimitiveSceneProxy && PrimitiveSceneProxy->CastsVolumetricTranslucentShadow(),
+					bCurrentlyRenderingSeparateTranslucency ? GScreenTextureUVScale : 1.0f
 				)
 			);
 
@@ -1090,6 +1096,8 @@ void FDeferredShadingSceneRenderer::RenderTranslucency(FRHICommandListImmediate&
 {
 	if (ShouldRenderTranslucency())
 	{
+		static const auto CVar = IConsoleManager::Get().FindTConsoleVariableDataFloat(TEXT("r.SeparateTranslucencyScreenPercentage"));
+		GScreenTextureUVScale = 100.0f / CVar->GetValueOnRenderThread();
 
 		SCOPED_DRAW_EVENT(RHICmdList, Translucency);
 
diff --git a/Engine/Source/Runtime/Slate/Private/Framework/Application/SlateApplication.cpp b/Engine/Source/Runtime/Slate/Private/Framework/Application/SlateApplication.cpp
--- a/Engine/Source/Runtime/Slate/Private/Framework/Application/SlateApplication.cpp
+++ b/Engine/Source/Runtime/Slate/Private/Framework/Application/SlateApplication.cpp
@@ -1710,6 +1710,29 @@ bool FSlateApplication::CanDisplayWindows() const
 }
 
 
+EUINavigation FSlateApplication::GetNavigationDirectionFromKey( const FKeyEvent& InKeyEvent ) const
+{
+	
+	if ( NavigationConfig.Right.Contains( InKeyEvent.GetKey() ) )
+	{
+		return EUINavigation::Right;
+	}
+	else if ( NavigationConfig.Left.Contains( InKeyEvent.GetKey() ) )
+	{
+		return EUINavigation::Left;
+	}
+	else if ( NavigationConfig.Up.Contains( InKeyEvent.GetKey() ) )
+	{
+		return EUINavigation::Up;
+	}
+	else if ( NavigationConfig.Down.Contains( InKeyEvent.GetKey() ) )
+	{
+		return EUINavigation::Down;
+	}
+	
+	return EUINavigation::Invalid;
+}
+
 void FSlateApplication::AddModalWindow( TSharedRef<SWindow> InSlateWindow, const TSharedPtr<const SWidget> InParentWidget, bool bSlowTaskWindow )
 {
 	if( !CanAddModalWindow() )
@@ -5462,6 +5485,11 @@ void FSlateApplication::ProcessApplicationActivationEvent( bool InAppActivated )
 }
 
 
+void FSlateApplication::SetNavigationConfig( FNavigationConfig&& Config )
+{
+	NavigationConfig = MoveTemp( Config );
+}
+
 bool FSlateApplication::OnConvertibleLaptopModeChanged()
 {
 	EConvertibleLaptopMode NewMode = FPlatformMisc::GetConvertibleLaptopMode();
diff --git a/Engine/Source/Runtime/Slate/Private/Widgets/Input/SEditableText.cpp b/Engine/Source/Runtime/Slate/Private/Widgets/Input/SEditableText.cpp
--- a/Engine/Source/Runtime/Slate/Private/Widgets/Input/SEditableText.cpp
+++ b/Engine/Source/Runtime/Slate/Private/Widgets/Input/SEditableText.cpp
@@ -1598,6 +1598,12 @@ const FSlateBrush* SEditableText::GetFocusBrush() const
 }
 
 
+bool SEditableText::IsInteractable() const
+{
+	return IsEnabled();
+}
+
+
 void SEditableText::SelectText( const int32 InOldCaretPosition )
 {
 	if( InOldCaretPosition != CaretPosition )
diff --git a/Engine/Source/Runtime/Slate/Private/Widgets/Layout/SScaleBox.cpp b/Engine/Source/Runtime/Slate/Private/Widgets/Layout/SScaleBox.cpp
--- a/Engine/Source/Runtime/Slate/Private/Widgets/Layout/SScaleBox.cpp
+++ b/Engine/Source/Runtime/Slate/Private/Widgets/Layout/SScaleBox.cpp
@@ -14,6 +14,7 @@ void SScaleBox::Construct( const SScaleBox::FArguments& InArgs )
 	Stretch = InArgs._Stretch;
 	StretchDirection = InArgs._StretchDirection;
 	UserSpecifiedScale = InArgs._UserSpecifiedScale;
+	IgnoreInheritedScale = InArgs._IgnoreInheritedScale;
 
 	ChildSlot
 	.HAlign(InArgs._HAlign)
@@ -78,6 +79,11 @@ void SScaleBox::OnArrangeChildren( const FGeometry& AllottedGeometry, FArrangedC
 			}
 		}
 
+		if (IgnoreInheritedScale.Get(false) && AllottedGeometry.Scale != 0)
+		{
+			FinalScale /= AllottedGeometry.Scale;
+		}
+
 		FVector2D FinalOffset(0, 0);
 
 		// If we're just filling, there's no scale applied, we're just filling the area.
@@ -158,6 +164,11 @@ void SScaleBox::SetUserSpecifiedScale(float InUserSpecifiedScale)
 	UserSpecifiedScale = InUserSpecifiedScale;
 }
 
+void SScaleBox::SetIgnoreInheritedScale(bool InIgnoreInheritedScale)
+{
+	IgnoreInheritedScale = InIgnoreInheritedScale;
+}
+
 float SScaleBox::GetRelativeLayoutScale(const FSlotBase& Child) const
 {
 	const EStretch::Type CurrentStretch = Stretch.Get();
diff --git a/Engine/Source/Runtime/Slate/Public/Framework/Application/NavigationConfig.h b/Engine/Source/Runtime/Slate/Public/Framework/Application/NavigationConfig.h
new file mode 100644
--- /dev/null
+++ b/Engine/Source/Runtime/Slate/Public/Framework/Application/NavigationConfig.h
@@ -0,0 +1,42 @@
+// Copyright 1998-2015 Epic Games, Inc. All Rights Reserved.
+
+#pragma once
+
+/** This structure is used to control which FKeys should move focus */
+struct FNavigationConfig
+{
+	/** Keys that should move focus left */
+	TSet<FKey> Left;
+
+	/** Keys that should move focus right */
+	TSet<FKey> Right;
+
+	/** Keys that should move focus up*/
+	TSet<FKey> Up;
+
+	/** Keys that should move focus down*/
+	TSet<FKey> Down;
+
+	FNavigationConfig()
+	{
+		// left
+		Left.Emplace( EKeys::Left );
+		Left.Emplace( EKeys::Gamepad_DPad_Left );
+		Left.Emplace( EKeys::Gamepad_LeftStick_Left );
+		
+		// right
+		Right.Emplace( EKeys::Right );
+		Right.Emplace( EKeys::Gamepad_DPad_Right );
+		Right.Emplace( EKeys::Gamepad_LeftStick_Right  );
+
+		// up
+		Up.Emplace( EKeys::Up );
+		Up.Emplace( EKeys::Gamepad_DPad_Up );
+		Up.Emplace( EKeys::Gamepad_LeftStick_Up );
+
+		// down
+		Down.Emplace( EKeys::Down );
+		Down.Emplace( EKeys::Gamepad_DPad_Down );
+		Down.Emplace( EKeys::Gamepad_LeftStick_Down );
+	}
+};
\ No newline at end of file
diff --git a/Engine/Source/Runtime/Slate/Public/Framework/Application/SlateApplication.h b/Engine/Source/Runtime/Slate/Public/Framework/Application/SlateApplication.h
--- a/Engine/Source/Runtime/Slate/Public/Framework/Application/SlateApplication.h
+++ b/Engine/Source/Runtime/Slate/Public/Framework/Application/SlateApplication.h
@@ -6,6 +6,7 @@
 #include "MenuStack.h"
 #include "SlateDelegates.h"
 #include "SlateApplicationBase.h"
+#include "NavigationConfig.h"
 
 
 class SToolTip;
@@ -82,7 +83,6 @@ class SLATE_API FPopupSupport
 
 };
 
-
 class SLATE_API FSlateApplication
 	: public FSlateApplicationBase
 	, public FGenericApplicationMessageHandler
@@ -221,6 +221,8 @@ public:
 
 	/** Returns true if this slate application is ready to display windows. */
 	bool CanDisplayWindows() const;
+
+	virtual EUINavigation GetNavigationDirectionFromKey( const FKeyEvent& InKeyEvent ) const override;
 	
 	/**
 	 * Adds a modal window to the application.  
@@ -941,6 +943,8 @@ public:
 
 public:
 
+	void SetNavigationConfig( FNavigationConfig&& Config );
+
 	/** Called when the slate application is being shut down. */
 	void OnShutdown();
 
@@ -1649,4 +1653,7 @@ private:
 
 	TMap< const ILayoutCache*, FCacheElementPools* > CachedElementLists;
 	TArray< FCacheElementPools* > ReleasedCachedElementLists;
+
+	/** Configured fkeys to control navigation */
+	FNavigationConfig NavigationConfig;
 };
diff --git a/Engine/Source/Runtime/Slate/Public/Widgets/Input/SEditableText.h b/Engine/Source/Runtime/Slate/Public/Widgets/Input/SEditableText.h
--- a/Engine/Source/Runtime/Slate/Public/Widgets/Input/SEditableText.h
+++ b/Engine/Source/Runtime/Slate/Public/Widgets/Input/SEditableText.h
@@ -415,6 +415,7 @@ protected:
 	virtual FReply OnMouseButtonDoubleClick( const FGeometry& InMyGeometry, const FPointerEvent& InMouseEvent ) override;
 	virtual FCursorReply OnCursorQuery( const FGeometry& MyGeometry, const FPointerEvent& CursorEvent ) const override;
 	virtual const FSlateBrush* GetFocusBrush() const;
+	virtual bool IsInteractable() const override;
 	//~ End SWidget Interface
 
 protected:
diff --git a/Engine/Source/Runtime/Slate/Public/Widgets/Layout/SScaleBox.h b/Engine/Source/Runtime/Slate/Public/Widgets/Layout/SScaleBox.h
--- a/Engine/Source/Runtime/Slate/Public/Widgets/Layout/SScaleBox.h
+++ b/Engine/Source/Runtime/Slate/Public/Widgets/Layout/SScaleBox.h
@@ -69,6 +69,7 @@ public:
 	, _StretchDirection(EStretchDirection::Both)
 	, _Stretch(EStretch::None)
 	, _UserSpecifiedScale(1.0f)
+	, _IgnoreInheritedScale(false)
 	{}
 		/** Slot for this designers content (optional) */
 		SLATE_DEFAULT_SLOT(FArguments, Content)
@@ -88,6 +89,9 @@ public:
 		/** Optional scale that can be specified by the User */
 		SLATE_ATTRIBUTE(float, UserSpecifiedScale)
 
+		/** Undo any inherited scale factor before applying this scale box's scale */
+		SLATE_ATTRIBUTE(bool, IgnoreInheritedScale)
+
 	SLATE_END_ARGS()
 
 	/** Constructor */
@@ -121,6 +125,9 @@ public:
 
 	/** See UserSpecifiedScale argument */
 	void SetUserSpecifiedScale(float InUserSpecifiedScale);
+
+	/** Set IgnoreInheritedScale argument */
+	void SetIgnoreInheritedScale(bool InIgnoreInheritedScale);
 	
 protected:
 	virtual float GetRelativeLayoutScale(const FSlotBase& Child) const override;
@@ -134,4 +141,7 @@ private:
 
 	/** Optional scale that can be specified by the User */
 	TAttribute<float> UserSpecifiedScale;
+
+	/** Optional bool to ignore the inherited scale */
+	TAttribute<bool> IgnoreInheritedScale;
 };
diff --git a/Engine/Source/Runtime/SlateCore/Private/Rendering/DrawElements.cpp b/Engine/Source/Runtime/SlateCore/Private/Rendering/DrawElements.cpp
--- a/Engine/Source/Runtime/SlateCore/Private/Rendering/DrawElements.cpp
+++ b/Engine/Source/Runtime/SlateCore/Private/Rendering/DrawElements.cpp
@@ -81,7 +81,7 @@ void FSlateDrawElement::MakeBox(
 	SCOPE_CYCLE_COUNTER(STAT_SlateDrawElementMakeTime)
 
 	// Ignore invalid rendering handles.
-	if ( !InRenderingHandle.Data.IsValid() )
+	if ( !InRenderingHandle.IsValid() )
 	{
 		return;
 	}
diff --git a/Engine/Source/Runtime/SlateCore/Private/Widgets/SWidget.cpp b/Engine/Source/Runtime/SlateCore/Private/Widgets/SWidget.cpp
--- a/Engine/Source/Runtime/SlateCore/Private/Widgets/SWidget.cpp
+++ b/Engine/Source/Runtime/SlateCore/Private/Widgets/SWidget.cpp
@@ -125,25 +125,13 @@ FReply SWidget::OnKeyDown( const FGeometry& MyGeometry, const FKeyEvent& InKeyEv
 {
 	if (SupportsKeyboardFocus())
 	{
+		EUINavigation Direction = FSlateApplicationBase::Get().GetNavigationDirectionFromKey( InKeyEvent );
 		// It's the left stick return a navigation request of the correct direction
-		if (InKeyEvent.GetKey() == EKeys::Right || InKeyEvent.GetKey() == EKeys::Gamepad_DPad_Right || InKeyEvent.GetKey() == EKeys::Gamepad_LeftStick_Right)
+		if ( Direction != EUINavigation::Invalid )
 		{
-			return FReply::Handled().SetNavigation(EUINavigation::Right);
+			return FReply::Handled().SetNavigation( Direction );
 		}
-		else if (InKeyEvent.GetKey() == EKeys::Left || InKeyEvent.GetKey() == EKeys::Gamepad_DPad_Left || InKeyEvent.GetKey() == EKeys::Gamepad_LeftStick_Left)
-		{
-			return FReply::Handled().SetNavigation(EUINavigation::Left);
-		}
-		else if (InKeyEvent.GetKey() == EKeys::Up || InKeyEvent.GetKey() == EKeys::Gamepad_DPad_Up || InKeyEvent.GetKey() == EKeys::Gamepad_LeftStick_Up)
-		{
-			return FReply::Handled().SetNavigation(EUINavigation::Up);
-		}
-		else if (InKeyEvent.GetKey() == EKeys::Down || InKeyEvent.GetKey() == EKeys::Gamepad_DPad_Down || InKeyEvent.GetKey() == EKeys::Gamepad_LeftStick_Down)
-		{
-			return FReply::Handled().SetNavigation(EUINavigation::Down);
-		}
-		// If the key was Tab, interpret as an attempt to move focus.
-		else if (InKeyEvent.GetKey() == EKeys::Tab)
+		else if ( InKeyEvent.GetKey() == EKeys::Tab )
 		{
 			//@TODO: Really these uses of input should be at a lower priority, only occurring if nothing else handled them
 			// For now this code prevents consuming them when some modifiers are held down, allowing some limited binding
diff --git a/Engine/Source/Runtime/SlateCore/Public/Application/SlateApplicationBase.h b/Engine/Source/Runtime/SlateCore/Public/Application/SlateApplicationBase.h
--- a/Engine/Source/Runtime/SlateCore/Public/Application/SlateApplicationBase.h
+++ b/Engine/Source/Runtime/SlateCore/Public/Application/SlateApplicationBase.h
@@ -189,6 +189,8 @@ public:
 	 */
 	virtual TSharedPtr< SWidget > GetKeyboardFocusedWidget( ) const = 0;
 
+	virtual EUINavigation GetNavigationDirectionFromKey( const FKeyEvent& InKeyEvent ) const = 0;
+
 	/**
 	* Gets the Widget that currently captures the mouse.
 	*
diff --git a/Engine/Source/Runtime/UMG/Private/Components/MultiLineEditableText.cpp b/Engine/Source/Runtime/UMG/Private/Components/MultiLineEditableText.cpp
--- a/Engine/Source/Runtime/UMG/Private/Components/MultiLineEditableText.cpp
+++ b/Engine/Source/Runtime/UMG/Private/Components/MultiLineEditableText.cpp
@@ -57,8 +57,10 @@ void UMultiLineEditableText::SynchronizeProperties()
 {
 	Super::SynchronizeProperties();
 
+	TAttribute<FText> HintTextBinding = OPTIONAL_BINDING(FText, HintText);
+
 	MyMultiLineEditableText->SetText(Text);
-//	MyMultiLineEditableText->SetHintText(HintText);
+	MyMultiLineEditableText->SetHintText(HintTextBinding);
 //	MyMultiLineEditableText->SetIsReadOnly(IsReadOnly);
 //	MyMultiLineEditableText->SetIsPassword(IsPassword);
 //	MyMultiLineEditableText->SetColorAndOpacity(ColorAndOpacity);
diff --git a/Engine/Source/Runtime/UMG/Private/Components/MultiLineEditableTextBox.cpp b/Engine/Source/Runtime/UMG/Private/Components/MultiLineEditableTextBox.cpp
--- a/Engine/Source/Runtime/UMG/Private/Components/MultiLineEditableTextBox.cpp
+++ b/Engine/Source/Runtime/UMG/Private/Components/MultiLineEditableTextBox.cpp
@@ -61,9 +61,12 @@ void UMultiLineEditableTextBox::SynchronizeProperties()
 {
 	Super::SynchronizeProperties();
 
+	TAttribute<FText> HintTextBinding = OPTIONAL_BINDING(FText, HintText);
+
 	MyEditableTextBlock->SetStyle(&WidgetStyle);
 	MyEditableTextBlock->SetText(Text);
-//	MyEditableTextBlock->SetHintText(HintText);
+	MyEditableTextBlock->SetHintText(HintTextBinding);
+
 //	MyEditableTextBlock->SetIsReadOnly(IsReadOnly);
 //	MyEditableTextBlock->SetIsPassword(IsPassword);
 //	MyEditableTextBlock->SetColorAndOpacity(ColorAndOpacity);
diff --git a/Engine/Source/Runtime/UMG/Private/Components/ScaleBox.cpp b/Engine/Source/Runtime/UMG/Private/Components/ScaleBox.cpp
--- a/Engine/Source/Runtime/UMG/Private/Components/ScaleBox.cpp
+++ b/Engine/Source/Runtime/UMG/Private/Components/ScaleBox.cpp
@@ -16,6 +16,7 @@ UScaleBox::UScaleBox(const FObjectInitializer& ObjectInitializer)
 	StretchDirection = EStretchDirection::Both;
 	Stretch = EStretch::ScaleToFit;
 	UserSpecifiedScale = 1.0f;
+	IgnoreInheritedScale = false;
 }
 
 void UScaleBox::ReleaseSlateResources(bool bReleaseChildren)
@@ -44,6 +45,7 @@ void UScaleBox::SynchronizeProperties()
 	MyScaleBox->SetStretchDirection(StretchDirection);
 	MyScaleBox->SetStretch(Stretch);
 	MyScaleBox->SetUserSpecifiedScale(UserSpecifiedScale);
+	MyScaleBox->SetIgnoreInheritedScale(IgnoreInheritedScale);
 }
 
 UClass* UScaleBox::GetSlotClass() const
diff --git a/Engine/Source/Runtime/UMG/Private/Slate/SMeshWidget.cpp b/Engine/Source/Runtime/UMG/Private/Slate/SMeshWidget.cpp
--- a/Engine/Source/Runtime/UMG/Private/Slate/SMeshWidget.cpp
+++ b/Engine/Source/Runtime/UMG/Private/Slate/SMeshWidget.cpp
@@ -5,12 +5,12 @@
 #include "SMeshWidget.h"
 #include "SlateMaterialBrush.h"
 #include "Runtime/SlateRHIRenderer/Public/Interfaces/ISlateRHIRendererModule.h"
-#include "SlateMeshData.h"
+#include "SlateVectorArtData.h"
 
 
 
 /** Populate OutSlateVerts and OutIndexes with data from this static mesh such that Slate can render it. */
-static void SlateMeshToSlateRenderData(const USlateMeshData& DataSource, TArray<FSlateVertex>& OutSlateVerts, TArray<SlateIndex>& OutIndexes)
+static void SlateMeshToSlateRenderData(const USlateVectorArtData& DataSource, TArray<FSlateVertex>& OutSlateVerts, TArray<SlateIndex>& OutIndexes)
 {
 	// Populate Index data
 	{
@@ -82,8 +82,10 @@ void SMeshWidget::Construct(const FArguments& Args)
 }
 
 static const FVector2D DontCare(FVector2D(64, 64));
-uint32 SMeshWidget::AddMesh(const USlateMeshData& InMeshData)
+uint32 SMeshWidget::AddMesh(USlateVectorArtData& InMeshData)
 {	
+	InMeshData.EnsureValidData();
+
 	FRenderData& NewRenderData = RenderData[RenderData.Add(FRenderData())];
 	UMaterialInterface* MaterialFromMesh = InMeshData.GetMaterial();
 	if (MaterialFromMesh != nullptr)
diff --git a/Engine/Source/Runtime/UMG/Private/Slate/SlateVectorArtData.cpp b/Engine/Source/Runtime/UMG/Private/Slate/SlateVectorArtData.cpp
new file mode 100644
--- /dev/null
+++ b/Engine/Source/Runtime/UMG/Private/Slate/SlateVectorArtData.cpp
@@ -0,0 +1,153 @@
+// Copyright 1998-2015 Epic Games, Inc. All Rights Reserved.
+
+#include "UMGPrivatePCH.h"
+#include "Engine/StaticMesh.h"
+#include "StaticMeshResources.h"
+#include "SlateVectorArtData.h"
+
+static void StaticMeshToSlateRenderData(const UStaticMesh& DataSource, TArray<FSlateMeshVertex>& OutSlateVerts, TArray<uint32>& OutIndexes)
+{
+	const FStaticMeshLODResources& LOD = DataSource.RenderData->LODResources[0];
+	const int32 NumSections = LOD.Sections.Num();
+	if (NumSections > 1)
+	{
+		UE_LOG(LogUMG, Warning, TEXT("StaticMesh %s has %d sections. SMeshWidget expects a static mesh with 1 section."), *DataSource.GetName(), NumSections);
+	}
+	else
+	{
+		// Populate Vertex Data
+		{
+			const uint32 NumVerts = LOD.PositionVertexBuffer.GetNumVertices();
+			OutSlateVerts.Empty();
+			OutSlateVerts.Reserve(NumVerts);
+
+			static const int32 MAX_SUPPORTED_UV_SETS = 6;
+			const int32 TexCoordsPerVertex = LOD.GetNumTexCoords();
+			if (TexCoordsPerVertex > MAX_SUPPORTED_UV_SETS)
+			{
+				UE_LOG(LogStaticMesh, Warning, TEXT("[%s] has %d UV sets; slate vertex data supports at most %d"), *DataSource.GetName(), TexCoordsPerVertex, MAX_SUPPORTED_UV_SETS);
+			}
+
+			for (uint32 i = 0; i < NumVerts; ++i)
+			{
+				// Copy Position
+				const FVector& Position = LOD.PositionVertexBuffer.VertexPosition(i);
+				
+				// Copy Color
+				FColor Color = (LOD.ColorVertexBuffer.GetNumVertices() > 0) ? LOD.ColorVertexBuffer.VertexColor(i) : FColor::White;
+				
+				// Copy all the UVs that we have, and as many as we can fit.
+				const FVector2D& UV0 = (TexCoordsPerVertex > 0) ? LOD.VertexBuffer.GetVertexUV(i, 0) : FVector2D(1, 1);
+
+				const FVector2D& UV1 = (TexCoordsPerVertex > 1) ? LOD.VertexBuffer.GetVertexUV(i, 1) : FVector2D(1, 1);
+
+				const FVector2D& UV2 = (TexCoordsPerVertex > 2) ? LOD.VertexBuffer.GetVertexUV(i, 2) : FVector2D(1, 1);
+
+				const FVector2D& UV3 = (TexCoordsPerVertex > 3) ? LOD.VertexBuffer.GetVertexUV(i, 3) : FVector2D(1, 1);
+
+				const FVector2D& UV4 = (TexCoordsPerVertex > 4) ? LOD.VertexBuffer.GetVertexUV(i, 4) : FVector2D(1, 1);
+
+				const FVector2D& UV5 = (TexCoordsPerVertex > 5) ? LOD.VertexBuffer.GetVertexUV(i, 5) : FVector2D(1, 1);
+
+				OutSlateVerts.Add(FSlateMeshVertex(
+					FVector2D(Position.X, Position.Y),
+					Color,
+					UV0,
+					UV1,
+					UV2,
+					UV3,
+					UV4,
+					UV5
+				));
+			}
+		}
+
+		// Populate Index data
+		{
+			FIndexArrayView SourceIndexes = LOD.IndexBuffer.GetArrayView();
+			const int32 NumIndexes = SourceIndexes.Num();
+			OutIndexes.Empty();
+			OutIndexes.Reserve(NumIndexes);
+			for (int32 i = 0; i < NumIndexes; ++i)
+			{
+				OutIndexes.Add(SourceIndexes[i]);
+			}
+
+
+			// Sort the index buffer such that verts are drawn in Z-order.
+			// Assume that all triangles are coplanar with Z == SomeValue.
+			ensure(NumIndexes % 3 == 0);
+			for (int32 a = 0; a < NumIndexes; a += 3)
+			{
+				for (int32 b = 0; b < NumIndexes; b += 3)
+				{
+					const float VertADepth = LOD.PositionVertexBuffer.VertexPosition(OutIndexes[a]).Z;
+					const float VertBDepth = LOD.PositionVertexBuffer.VertexPosition(OutIndexes[b]).Z;
+					if ( VertADepth < VertBDepth )
+					{
+						// Swap the order in which triangles will be drawn
+						Swap(OutIndexes[a + 0], OutIndexes[b + 0]);
+						Swap(OutIndexes[a + 1], OutIndexes[b + 1]);
+						Swap(OutIndexes[a + 2], OutIndexes[b + 2]);
+					}
+				}
+			}
+		}
+	}
+}
+
+
+USlateVectorArtData::USlateVectorArtData(const FObjectInitializer& ObjectInitializer)
+	: Super(ObjectInitializer)
+{
+
+}
+
+const TArray<FSlateMeshVertex>& USlateVectorArtData::GetVertexData() const
+{
+	return VertexData;
+}
+
+const TArray<uint32>& USlateVectorArtData::GetIndexData() const
+{
+	return IndexData;
+}
+
+UMaterialInterface* USlateVectorArtData::GetMaterial() const
+{
+	return Material;
+}
+
+UMaterialInstanceDynamic* USlateVectorArtData::ConvertToMaterialInstanceDynamic()
+{
+	UMaterialInstanceDynamic* NewMID = UMaterialInstanceDynamic::Create(Material, this);
+	Material = NewMID;
+	return NewMID;
+}
+
+void USlateVectorArtData::EnsureValidData()
+{
+#if WITH_EDITORONLY_DATA
+	if (VertexData.Num() == 0 || IndexData.Num() == 0 || Material == nullptr)
+	{
+		InitFromStaticMesh(*MeshAsset);
+	}
+#endif
+}
+
+void USlateVectorArtData::PreSave()
+{
+	Super::PreSave();
+	EnsureValidData();
+}
+
+#if WITH_EDITORONLY_DATA
+void USlateVectorArtData::InitFromStaticMesh(const UStaticMesh& InSourceMesh)
+{
+	Material = InSourceMesh.GetMaterial(0);
+
+	ensureMsgf(Material != nullptr, TEXT("USlateVectorArtData::InitFromStaticMesh() expected %s to have a material assigned."), *InSourceMesh.GetFullName());
+
+	StaticMeshToSlateRenderData(InSourceMesh, VertexData, IndexData);
+}
+#endif
\ No newline at end of file
diff --git a/Engine/Source/Runtime/UMG/Public/Components/MultiLineEditableText.h b/Engine/Source/Runtime/UMG/Public/Components/MultiLineEditableText.h
--- a/Engine/Source/Runtime/UMG/Public/Components/MultiLineEditableText.h
+++ b/Engine/Source/Runtime/UMG/Public/Components/MultiLineEditableText.h
@@ -22,6 +22,13 @@ public:
 	UPROPERTY(EditAnywhere, Category=Content, meta=(MultiLine="true"))
 	FText Text;
 
+	/** Hint text that appears when there is no text in the text box */
+	UPROPERTY(EditAnywhere, Category=Content, meta=(MultiLine="true"))
+	FText HintText;
+
+	/** A bindable delegate to allow logic to drive the hint text of the widget */
+	UPROPERTY()
+	FGetText HintTextDelegate;
 public:
 
 	/** The style */
diff --git a/Engine/Source/Runtime/UMG/Public/Components/MultiLineEditableTextBox.h b/Engine/Source/Runtime/UMG/Public/Components/MultiLineEditableTextBox.h
--- a/Engine/Source/Runtime/UMG/Public/Components/MultiLineEditableTextBox.h
+++ b/Engine/Source/Runtime/UMG/Public/Components/MultiLineEditableTextBox.h
@@ -23,6 +23,13 @@ public:
 	UPROPERTY(EditAnywhere, Category=Content, meta=(MultiLine="true"))
 	FText Text;
 
+	/** Hint text that appears when there is no text in the text box */
+	UPROPERTY(EditAnywhere, Category=Content, meta=(MultiLine="true"))
+	FText HintText;
+
+	/** A bindable delegate to allow logic to drive the hint text of the widget */
+	UPROPERTY()
+	FGetText HintTextDelegate;
 public:
 
 	/** The style */
diff --git a/Engine/Source/Runtime/UMG/Public/Components/ScaleBox.h b/Engine/Source/Runtime/UMG/Public/Components/ScaleBox.h
--- a/Engine/Source/Runtime/UMG/Public/Components/ScaleBox.h
+++ b/Engine/Source/Runtime/UMG/Public/Components/ScaleBox.h
@@ -31,11 +31,16 @@ public:
 	UPROPERTY(EditAnywhere, BlueprintReadOnly, Category="Stretching")
 	float UserSpecifiedScale;
 
+	/** Optional bool to ignore the inherited scale. Applies inverse scaling to counteract parents before applying the local scale operation. */
+	UPROPERTY(EditAnywhere, BlueprintReadOnly, Category = "Stretching")
+	bool IgnoreInheritedScale;
+
 public:
 
 	//TODO UMG Add Set Stretch
 	//TODO UMG Add Set Stretch DIrection
 	//TODO UMG Add Set UserSpecifiedScale
+	//TODO UMG Add Set IgnoreInheritedScale
 	
 public:
 
diff --git a/Engine/Source/Runtime/UMG/Public/Slate/SMeshWidget.h b/Engine/Source/Runtime/UMG/Public/Slate/SMeshWidget.h
--- a/Engine/Source/Runtime/UMG/Public/Slate/SMeshWidget.h
+++ b/Engine/Source/Runtime/UMG/Public/Slate/SMeshWidget.h
@@ -2,7 +2,7 @@
 
 #pragma once
 
-class USlateMeshData;
+class USlateVectorArtData;
 
 
 /**
@@ -17,7 +17,7 @@ public:
 		: _MeshData(nullptr)
 	{}
 		/** The StaticMesh asset that should be drawn. */
-		SLATE_ARGUMENT(USlateMeshData*, MeshData)
+		SLATE_ARGUMENT(USlateVectorArtData*, MeshData)
 	SLATE_END_ARGS()
 
 	void Construct(const FArguments& Args);
@@ -27,7 +27,7 @@ public:
 	 *
 	 * @return the Index of the mesh data that was added; cache this value for use with @see FRenderRun.
 	 */
-	uint32 AddMesh(const USlateMeshData& InMeshData);
+	uint32 AddMesh(USlateVectorArtData& InMeshData);
 
 	/** Discard any previous runs and reserve space for new render runs if needed. */
 	void ClearRuns(int32 NumRuns);
diff --git a/Engine/Source/Runtime/UMG/Public/Slate/SlateVectorArtData.h b/Engine/Source/Runtime/UMG/Public/Slate/SlateVectorArtData.h
new file mode 100644
--- /dev/null
+++ b/Engine/Source/Runtime/UMG/Public/Slate/SlateVectorArtData.h
@@ -0,0 +1,116 @@
+// Cop// Copyright 1998-2015 Epic Games, Inc. All Rights Reserved.
+
+#pragma once
+
+#include "SlateVectorArtData.generated.h"
+
+class UStaticMesh;
+class UMaterialInterface;
+
+USTRUCT()
+struct FSlateMeshVertex
+{
+	GENERATED_USTRUCT_BODY()
+
+	static const int32 MaxNumUVs = 6;
+
+	FSlateMeshVertex()
+	{
+	}
+
+	FSlateMeshVertex(
+		  FVector2D InPos
+		, FColor InColor
+		, FVector2D InUV0
+		, FVector2D InUV1
+		, FVector2D InUV2
+		, FVector2D InUV3
+		, FVector2D InUV4
+		, FVector2D InUV5
+		)
+		: Position(InPos)
+		, Color(InColor)
+		, UV0(InUV0)
+		, UV1(InUV1)
+		, UV2(InUV2)
+		, UV3(InUV3)
+		, UV4(InUV4)
+		, UV5(InUV5)	
+	{
+	}
+
+	UPROPERTY()
+	FVector2D Position;
+
+	UPROPERTY()
+	FColor Color;
+
+	UPROPERTY()
+	FVector2D UV0;
+
+	UPROPERTY()
+	FVector2D UV1;
+
+	UPROPERTY()
+	FVector2D UV2;
+
+	UPROPERTY()
+	FVector2D UV3;
+
+	UPROPERTY()
+	FVector2D UV4;
+
+	UPROPERTY()
+	FVector2D UV5;
+};
+
+/**
+ * Turn static mesh data into Slate's simple vector art format.
+ */
+UCLASS()
+class UMG_API USlateVectorArtData : public UObject
+{
+	GENERATED_UCLASS_BODY()
+
+public:
+	/** Access the slate vertexes. */
+	const TArray<FSlateMeshVertex>& GetVertexData() const;
+	
+	/** Access the indexes for the order in which to draw the vertexes. */
+	const TArray<uint32>& GetIndexData() const;
+	
+	/** Material to be used with the specified vector art data. */
+	UMaterialInterface* GetMaterial() const;
+	
+	/** Convert the material into an MID and get a pointer to the MID so that parameters can be set on it. */
+	UMaterialInstanceDynamic* ConvertToMaterialInstanceDynamic();
+
+	/** Convert the static mesh data into slate vector art on demand. Does nothing in a cooked build. */
+	void EnsureValidData();
+
+private:
+	// ~ UObject Interface
+	virtual void PreSave() override;
+	// ~ UObject Interface
+
+#if WITH_EDITORONLY_DATA
+	/** Does the actual work of converting mesh data into slate vector art */
+	void InitFromStaticMesh(const UStaticMesh& InSourceMesh);
+
+	/** The mesh data asset from which the vector art is sourced */
+	UPROPERTY(EditAnywhere, Category="Vector Art" )
+	UStaticMesh* MeshAsset;
+#endif
+
+	/** @see GetVertexData() */
+	UPROPERTY()
+	TArray<FSlateMeshVertex> VertexData;
+
+	/** @see GetIndexData() */
+	UPROPERTY()
+	TArray<uint32> IndexData;
+
+	/** @see GetMaterial() */
+	UPROPERTY()
+	UMaterialInterface* Material;
+};
\ No newline at end of file
diff --git a/Engine/Source/Runtime/WebBrowser/Private/WebBrowserWindow.cpp b/Engine/Source/Runtime/WebBrowser/Private/WebBrowserWindow.cpp
--- a/Engine/Source/Runtime/WebBrowser/Private/WebBrowserWindow.cpp
+++ b/Engine/Source/Runtime/WebBrowser/Private/WebBrowserWindow.cpp
@@ -137,7 +137,13 @@ TSharedRef<SWidget> FWebBrowserWindow::CreateWidget(TAttribute<FVector2D> InView
 
 FSlateShaderResource* FWebBrowserWindow::GetTexture(bool bIsPopup)
 {
-	if (UpdatableTextures[bIsPopup?1:0] != nullptr && IsInitialized())
+	if (!bIsPopup && UpdatableTextures[0] == nullptr && FSlateApplication::IsInitialized())
+	{
+		// SViewport renders a black quad over the entire view if we return nullptr. Return an empty texture instead.
+		UpdatableTextures[0] = FSlateApplication::Get().GetRenderer()->CreateUpdatableTexture(1, 1);
+	}
+
+	if (UpdatableTextures[bIsPopup ? 1 : 0] != nullptr)
 	{
 		return UpdatableTextures[bIsPopup?1:0]->GetSlateResource();
 	}
