Copying //UE4/Orion-Staging to Dev-Main (//UE4/Dev-Main)

diff --git a/Engine/Source/Runtime/Core/Public/Misc/AutomationTest.h b/Engine/Source/Runtime/Core/Public/Misc/AutomationTest.h
--- a/Engine/Source/Runtime/Core/Public/Misc/AutomationTest.h
+++ b/Engine/Source/Runtime/Core/Public/Misc/AutomationTest.h
@@ -106,7 +106,7 @@ public:
 	 * @param	InTestFlag - Test flags
 	 * @param	InParameterName - optional parameter. e.g. asset name
 	 */
-	FAutomationTestInfo( const FString& InDisplayName, const FString& InTestName, const uint8 InTestFlags, const int32 InNumParticipantsRequired, const FString& InParameterName = FString() )
+	FAutomationTestInfo( const FString& InDisplayName, const FString& InTestName, const uint32 InTestFlags, const int32 InNumParticipantsRequired, const FString& InParameterName = FString() )
 		: DisplayName( InDisplayName )
 		, TestName( InTestName )
 		, TestParameter( InParameterName )
@@ -177,7 +177,7 @@ public:
 	 *
 	 * @return the test type.
 	 */
-	const uint8 GetTestFlags() const
+	const uint32 GetTestFlags() const
 	{
 		return TestFlags;
 	}
@@ -280,7 +280,7 @@ private:
 	FString TestParameter;
 
 	/** The test flags. */
-	uint8 TestFlags;
+	uint32 TestFlags;
 
 	/** The number of participants this test requires */
 	uint32 NumParticipantsRequired;
diff --git a/Engine/Source/Runtime/CoreUObject/Private/UObject/UObjectGlobals.cpp b/Engine/Source/Runtime/CoreUObject/Private/UObject/UObjectGlobals.cpp
--- a/Engine/Source/Runtime/CoreUObject/Private/UObject/UObjectGlobals.cpp
+++ b/Engine/Source/Runtime/CoreUObject/Private/UObject/UObjectGlobals.cpp
@@ -1505,20 +1505,22 @@ FName MakeUniqueObjectName( UObject* Parent, UClass* Class, FName BaseName/*=NAM
 }
 
 /**
- * Given an actor label string, generates an FName that can be used as an object name for that label.
- * The generated name isn't guaranteed to be unique.  If the object's current name is already satisfactory, then
- * that name will be returned.
+ * Given a display label string, generates an FName slug that is a valid FName for that label.
+ * If the object's current name is already satisfactory, then that name will be returned.
+ * For example, "[MyObject]: Object Label" becomes "MyObjectObjectLabel" FName slug.
+ * 
+ * Note: The generated name isn't guaranteed to be unique.
  *
- * @param	ActorLabel	The label string to convert to an FName
- * @param	CurrentObjectName	The object's current name, or NAME_None if it has no name yet
+ * @param DisplayLabel The label string to convert to an FName
+ * @param CurrentObjectName The object's current name, or NAME_None if it has no name yet
  *
- * @return	The generated actor object name
+ * @return	The generated object name
  */
-FName MakeObjectNameFromActorLabel( const FString& InActorLabel, const FName CurrentObjectName )
+FName MakeObjectNameFromDisplayLabel(const FString& DisplayLabel, const FName CurrentObjectName)
 {
-	FString GeneratedName = InActorLabel;
+	FString GeneratedName = DisplayLabel;
 
-	// Convert the actor label, which may consist of just about any possible character, into a
+	// Convert the display label, which may consist of just about any possible character, into a
 	// suitable name for a UObject (remove whitespace, certain symbols, etc.)
 	{
 		for( int32 BadCharacterIndex = 0; BadCharacterIndex < ARRAY_COUNT( INVALID_OBJECTNAME_CHARACTERS ) - 1; ++BadCharacterIndex )
diff --git a/Engine/Source/Runtime/CoreUObject/Public/UObject/UObjectGlobals.h b/Engine/Source/Runtime/CoreUObject/Public/UObject/UObjectGlobals.h
--- a/Engine/Source/Runtime/CoreUObject/Public/UObject/UObjectGlobals.h
+++ b/Engine/Source/Runtime/CoreUObject/Public/UObject/UObjectGlobals.h
@@ -408,16 +408,18 @@ COREUOBJECT_API void IncrementalPurgeGarbage( bool bUseTimeLimit, float TimeLimi
 COREUOBJECT_API FName MakeUniqueObjectName( UObject* Outer, UClass* Class, FName BaseName=NAME_None );
 
 /**
- * Given an actor label string, generates an FName that can be used as an object name for that label.
- * The generated name isn't guaranteed to be unique.  If the object's current name is already satisfactory, then
- * that name will be returned.
+ * Given a display label string, generates an FName slug that is a valid FName for that label.
+ * If the object's current name is already satisfactory, then that name will be returned.
+ * For example, "[MyObject]: Object Label" becomes "MyObjectObjectLabel" FName slug.
+ * 
+ * Note: The generated name isn't guaranteed to be unique.
  *
- * @param	ActorLabel	The label string to convert to an FName
- * @param	CurrentObjectName	The object's current name, or NAME_None if it has no name yet
+ * @param DisplayLabel The label string to convert to an FName
+ * @param CurrentObjectName The object's current name, or NAME_None if it has no name yet
  *
- * @return	The generated actor object name
+ * @return	The generated object name
  */
-COREUOBJECT_API FName MakeObjectNameFromActorLabel( const FString& InActorLabel, const FName CurrentObjectName );
+COREUOBJECT_API FName MakeObjectNameFromDisplayLabel(const FString& DisplayLabel, const FName CurrentObjectName);
 
 /**
  * Returns whether an object is referenced, not counting the one
diff --git a/Engine/Source/Runtime/Engine/Classes/Curves/CurveBase.h b/Engine/Source/Runtime/Engine/Classes/Curves/CurveBase.h
--- a/Engine/Source/Runtime/Engine/Classes/Curves/CurveBase.h
+++ b/Engine/Source/Runtime/Engine/Classes/Curves/CurveBase.h
@@ -247,9 +247,10 @@ public:
 	 * Finds a key a the specified time.
 	 *
 	 * @param KeyTime The time at which to find the key.
+	 * @param KeyTimeTolerance The key time tolerance to use for equality.
 	 * @return A handle to the key, or invalid key handle if not found.
 	 */
-	FKeyHandle FindKey(float KeyTime) const;
+	FKeyHandle FindKey(float KeyTime, float KeyTimeTolerance = KINDA_SMALL_NUMBER) const;
 
 	/**
 	 * Get a key.
@@ -543,7 +544,7 @@ public:
 	float GetKeyTime(FKeyHandle KeyHandle) const;
 
 	/** Finds a key a the specified time */
-	FKeyHandle FindKey( float KeyTime ) const;
+	FKeyHandle FindKey( float KeyTime, float KeyTimeTolerance = KINDA_SMALL_NUMBER ) const;
 
 	/** Set the value of the specified key */
 	void SetKeyValue(FKeyHandle KeyHandle, float NewValue, bool bAutoSetTangents=true);
diff --git a/Engine/Source/Runtime/Engine/Classes/GameFramework/WorldSettings.h b/Engine/Source/Runtime/Engine/Classes/GameFramework/WorldSettings.h
--- a/Engine/Source/Runtime/Engine/Classes/GameFramework/WorldSettings.h
+++ b/Engine/Source/Runtime/Engine/Classes/GameFramework/WorldSettings.h
@@ -466,19 +466,20 @@ class ENGINE_API AWorldSettings : public AInfo, public IInterface_AssetUserData
 	 * Normally 1 - scales real time passage.
 	 * Warning - most use cases should use GetEffectiveTimeDilation() instead of reading from this directly
 	 */
-	UPROPERTY(replicated)
+	UPROPERTY(transient, replicated)
 	float TimeDilation;
 
-	// additional TimeDilation used by Matinee slomo track
-	UPROPERTY(replicated)
+	// Additional time dilation used by Matinee (or Sequencer) slomo track.  Transient because this is often 
+	// temporarily modified by the editor when previewing slow motion effects, yet we don't want it saved or loaded from level packages.
+	UPROPERTY(transient, replicated)
 	float MatineeTimeDilation;
 
 	// Additional TimeDilation used to control demo playback speed
-	UPROPERTY()
+	UPROPERTY(transient)
 	float DemoPlayTimeDilation;
 
 	// If paused, FName of person pausing the game.
-	UPROPERTY(replicated)
+	UPROPERTY(transient, replicated)
 	class APlayerState* Pauser;
 
 	/** when this flag is set, more time is allocated to background loading (replicated) */
diff --git a/Engine/Source/Runtime/Engine/Private/ActorEditor.cpp b/Engine/Source/Runtime/Engine/Private/ActorEditor.cpp
--- a/Engine/Source/Runtime/Engine/Private/ActorEditor.cpp
+++ b/Engine/Source/Runtime/Engine/Private/ActorEditor.cpp
@@ -635,7 +635,7 @@ void AActor::SetActorLabelInternal( const FString& NewActorLabelDirty, bool bMak
 	{
 		// Generate an object name for the actor's label
 		const FName OldActorName = GetFName();
-		FName NewActorName = MakeObjectNameFromActorLabel( GetActorLabel(), OldActorName );
+		FName NewActorName = MakeObjectNameFromDisplayLabel( GetActorLabel(), OldActorName );
 
 		// Has anything changed?
 		if( OldActorName != NewActorName )
diff --git a/Engine/Source/Runtime/Engine/Private/CurveBase.cpp b/Engine/Source/Runtime/Engine/Private/CurveBase.cpp
--- a/Engine/Source/Runtime/Engine/Private/CurveBase.cpp
+++ b/Engine/Source/Runtime/Engine/Private/CurveBase.cpp
@@ -219,7 +219,7 @@ void FNameCurve::DeleteKey(FKeyHandle KeyHandle)
 }
 
 
-FKeyHandle FNameCurve::FindKey(float KeyTime) const
+FKeyHandle FNameCurve::FindKey(float KeyTime, float KeyTimeTolerance) const
 {
 	int32 Start = 0;
 	int32 End = Keys.Num()-1;
@@ -230,7 +230,7 @@ FKeyHandle FNameCurve::FindKey(float KeyTime) const
 		int32 TestPos = Start + (End-Start) / 2;
 		float TestKeyTime = Keys[TestPos].Time;
 
-		if (TestKeyTime == KeyTime)
+		if (FMath::IsNearlyEqual(TestKeyTime, KeyTime, KeyTimeTolerance))
 		{
 			return GetKeyHandle(TestPos);
 		}
@@ -732,7 +732,7 @@ float FRichCurve::GetKeyTime(FKeyHandle KeyHandle) const
 }
 
 
-FKeyHandle FRichCurve::FindKey(float KeyTime) const
+FKeyHandle FRichCurve::FindKey(float KeyTime, float KeyTimeTolerance) const
 {
 	int32 Start = 0;
 	int32 End = Keys.Num()-1;
@@ -743,7 +743,7 @@ FKeyHandle FRichCurve::FindKey(float KeyTime) const
 		int32 TestPos = Start + (End-Start) / 2;
 		float TestKeyTime = Keys[TestPos].Time;
 
-		if (TestKeyTime == KeyTime)
+		if (FMath::IsNearlyEqual(TestKeyTime, KeyTime, KeyTimeTolerance))
 		{
 			return GetKeyHandle(TestPos);
 		}
diff --git a/Engine/Source/Runtime/Engine/Private/DataChannel.cpp b/Engine/Source/Runtime/Engine/Private/DataChannel.cpp
--- a/Engine/Source/Runtime/Engine/Private/DataChannel.cpp
+++ b/Engine/Source/Runtime/Engine/Private/DataChannel.cpp
@@ -24,6 +24,7 @@ DEFINE_LOG_CATEGORY(LogSecurity);
 DEFINE_LOG_CATEGORY_STATIC(LogNetPartialBunch, Warning, All);
 
 DECLARE_CYCLE_STAT(TEXT("ActorChan_ReceivedBunch"), Stat_ActorChanReceivedBunch, STATGROUP_Net);
+DECLARE_CYCLE_STAT(TEXT("ActorChan_CleanUp"), Stat_ActorChanCleanUp, STATGROUP_Net);
 
 extern FAutoConsoleVariable CVarDoReplicationContextString;
 
@@ -1521,6 +1522,8 @@ void UActorChannel::DestroyActorAndComponents()
 
 bool UActorChannel::CleanUp( const bool bForDestroy )
 {
+	SCOPE_CYCLE_COUNTER(Stat_ActorChanCleanUp);
+
 	checkf(Connection != nullptr, TEXT("UActorChannel::CleanUp: Connection is null!"));
 	checkf(Connection->Driver != nullptr, TEXT("UActorChannel::CleanUp: Connection->Driver is null!"));
 
diff --git a/Engine/Source/Runtime/Engine/Private/GameViewportClient.cpp b/Engine/Source/Runtime/Engine/Private/GameViewportClient.cpp
--- a/Engine/Source/Runtime/Engine/Private/GameViewportClient.cpp
+++ b/Engine/Source/Runtime/Engine/Private/GameViewportClient.cpp
@@ -2749,6 +2749,18 @@ bool UGameViewportClient::HandleToggleFullscreenCommand()
 		}
 	}
 
+	// Make sure the user's settings are updated after pressing Alt+Enter to toggle fullscreen.  Note
+	// that we don't need to "apply" the setting change, as we already did that above directly.
+	UGameEngine* GameEngine = Cast<UGameEngine>( GEngine );
+	if( GameEngine )
+	{
+		UGameUserSettings* UserSettings = GameEngine->GetGameUserSettings();
+		if( UserSettings != nullptr )
+		{
+			UserSettings->SetFullscreenMode( FullScreenMode );
+		}
+	}
+
 	FSystemResolution::RequestResolutionChange(GSystemResolution.ResX, GSystemResolution.ResY, FullScreenMode);
 	return true;
 }
diff --git a/Engine/Source/Runtime/Engine/Private/GameplayStatics.cpp b/Engine/Source/Runtime/Engine/Private/GameplayStatics.cpp
--- a/Engine/Source/Runtime/Engine/Private/GameplayStatics.cpp
+++ b/Engine/Source/Runtime/Engine/Private/GameplayStatics.cpp
@@ -703,6 +703,7 @@ UParticleSystemComponent* UGameplayStatics::SpawnEmitterAttached(UParticleSystem
 				PSC->SetRelativeLocationAndRotation(Location, Rotation);
 			}
 			PSC->SetRelativeScale3D(FVector(1.f));
+			PSC->RegisterComponentWithWorld(AttachToComponent->GetWorld());
 			PSC->ActivateSystem(true);
 		}
 	}
diff --git a/Engine/Source/Runtime/GameplayAbilities/Private/Abilities/GameplayAbility.cpp b/Engine/Source/Runtime/GameplayAbilities/Private/Abilities/GameplayAbility.cpp
--- a/Engine/Source/Runtime/GameplayAbilities/Private/Abilities/GameplayAbility.cpp
+++ b/Engine/Source/Runtime/GameplayAbilities/Private/Abilities/GameplayAbility.cpp
@@ -1476,6 +1476,8 @@ TArray<FActiveGameplayEffectHandle> UGameplayAbility::K2_ApplyGameplayEffectToTa
 
 TArray<FActiveGameplayEffectHandle> UGameplayAbility::ApplyGameplayEffectToTarget(const FGameplayAbilitySpecHandle Handle, const FGameplayAbilityActorInfo* ActorInfo, const FGameplayAbilityActivationInfo ActivationInfo, FGameplayAbilityTargetDataHandle Target, TSubclassOf<UGameplayEffect> GameplayEffectClass, float GameplayEffectLevel, int32 Stacks) const
 {
+	SCOPE_CYCLE_COUNTER(STAT_ApplyGameplayEffectToTarget);
+
 	TArray<FActiveGameplayEffectHandle> EffectHandles;
 
 	if (HasAuthority(&ActivationInfo) == false && UAbilitySystemGlobals::Get().ShouldPredictTargetGameplayEffects() == false)
diff --git a/Engine/Source/Runtime/GameplayAbilities/Private/AbilitySystemStats.cpp b/Engine/Source/Runtime/GameplayAbilities/Private/AbilitySystemStats.cpp
--- a/Engine/Source/Runtime/GameplayAbilities/Private/AbilitySystemStats.cpp
+++ b/Engine/Source/Runtime/GameplayAbilities/Private/AbilitySystemStats.cpp
@@ -24,4 +24,7 @@ DEFINE_STAT(STAT_AggregatorEvaluate);
 DEFINE_STAT(STAT_HasApplicationImmunityToSpec);
 DEFINE_STAT(STAT_HasMatchingGameplayTag);
 DEFINE_STAT(STAT_HandleGameplayCueNotifyStatic);
-DEFINE_STAT(STAT_HandleGameplayCueNotifyActor);
\ No newline at end of file
+DEFINE_STAT(STAT_HandleGameplayCueNotifyActor);
+DEFINE_STAT(STAT_ApplyGameplayEffectToTarget);
+DEFINE_STAT(STAT_OnActiveGameplayEffectAdded);
+DEFINE_STAT(STAT_OnActiveGameplayEffectRemoved);
diff --git a/Engine/Source/Runtime/GameplayAbilities/Private/GameplayCueSet.cpp b/Engine/Source/Runtime/GameplayAbilities/Private/GameplayCueSet.cpp
--- a/Engine/Source/Runtime/GameplayAbilities/Private/GameplayCueSet.cpp
+++ b/Engine/Source/Runtime/GameplayAbilities/Private/GameplayCueSet.cpp
@@ -48,7 +48,7 @@ void UGameplayCueSet::AddCues(const TArray<FGameplayCueReferencePair>& CuesToAdd
 			{
 				if (Data.GameplayCueTag == GameplayCueTag)
 				{
-					ABILITY_LOG(Warning, TEXT("AddGameplayCueData_Internal called for [%s,%s] when it already existed [%s,%s]. Skipping."), *GameplayCueTag.ToString(), *StringRef.ToString(), *Data.GameplayCueTag.ToString(), *Data.GameplayCueNotifyObj.ToString());
+					ABILITY_LOG(Verbose, TEXT("AddGameplayCueData_Internal called for [%s,%s] when it already existed [%s,%s]. Skipping."), *GameplayCueTag.ToString(), *StringRef.ToString(), *Data.GameplayCueTag.ToString(), *Data.GameplayCueNotifyObj.ToString());
 					bDupe = true;
 					break;
 				}
diff --git a/Engine/Source/Runtime/GameplayAbilities/Private/GameplayEffect.cpp b/Engine/Source/Runtime/GameplayAbilities/Private/GameplayEffect.cpp
--- a/Engine/Source/Runtime/GameplayAbilities/Private/GameplayEffect.cpp
+++ b/Engine/Source/Runtime/GameplayAbilities/Private/GameplayEffect.cpp
@@ -2334,6 +2334,8 @@ FActiveGameplayEffect* FActiveGameplayEffectsContainer::ApplyGameplayEffectSpec(
 /** This is called anytime a new ActiveGameplayEffect is added, on both client and server in all cases */
 void FActiveGameplayEffectsContainer::InternalOnActiveGameplayEffectAdded(FActiveGameplayEffect& Effect)
 {
+	SCOPE_CYCLE_COUNTER(STAT_OnActiveGameplayEffectAdded);
+
 	if (Effect.Spec.Def == nullptr)
 	{
 		ABILITY_LOG(Error, TEXT("FActiveGameplayEffectsContainer serialized new GameplayEffect with NULL Def!"));
@@ -2560,6 +2562,8 @@ bool FActiveGameplayEffectsContainer::InternalRemoveActiveGameplayEffect(int32 I
 /** Called by client and server: This does cleanup that has to happen whether the effect is being removed locally or due to replication */
 void FActiveGameplayEffectsContainer::InternalOnActiveGameplayEffectRemoved(const FActiveGameplayEffect& Effect, bool bInvokeGameplayCueEvents)
 {
+	SCOPE_CYCLE_COUNTER(STAT_OnActiveGameplayEffectRemoved);
+
 	if (Effect.Spec.Def)
 	{
 		// Remove our tag requirements from the dependency map
diff --git a/Engine/Source/Runtime/GameplayAbilities/Public/AbilitySystemStats.h b/Engine/Source/Runtime/GameplayAbilities/Public/AbilitySystemStats.h
--- a/Engine/Source/Runtime/GameplayAbilities/Public/AbilitySystemStats.h
+++ b/Engine/Source/Runtime/GameplayAbilities/Public/AbilitySystemStats.h
@@ -28,4 +28,7 @@ DECLARE_CYCLE_STAT_EXTERN(TEXT("Aggregator Evaluate"), STAT_AggregatorEvaluate,
 DECLARE_CYCLE_STAT_EXTERN(TEXT("Has Application Immunity To Spec"), STAT_HasApplicationImmunityToSpec, STATGROUP_AbilitySystem, );
 DECLARE_CYCLE_STAT_EXTERN(TEXT("Has Matching GameplayTag"), STAT_HasMatchingGameplayTag, STATGROUP_AbilitySystem, );
 DECLARE_CYCLE_STAT_EXTERN(TEXT("GameplayCueNotify Static"), STAT_HandleGameplayCueNotifyStatic, STATGROUP_AbilitySystem, );
-DECLARE_CYCLE_STAT_EXTERN(TEXT("GameplayCueNotify Actor"), STAT_HandleGameplayCueNotifyActor, STATGROUP_AbilitySystem, );
\ No newline at end of file
+DECLARE_CYCLE_STAT_EXTERN(TEXT("GameplayCueNotify Actor"), STAT_HandleGameplayCueNotifyActor, STATGROUP_AbilitySystem, );
+DECLARE_CYCLE_STAT_EXTERN(TEXT("ApplyGameplayEffectToTarget"), STAT_ApplyGameplayEffectToTarget, STATGROUP_AbilitySystem, );
+DECLARE_CYCLE_STAT_EXTERN(TEXT("ActiveGameplayEffect Added"), STAT_OnActiveGameplayEffectAdded, STATGROUP_AbilitySystem, );
+DECLARE_CYCLE_STAT_EXTERN(TEXT("ActiveGameplayEffect Removed"), STAT_OnActiveGameplayEffectRemoved, STATGROUP_AbilitySystem, );
\ No newline at end of file
diff --git a/Engine/Source/Runtime/GameplayTags/Private/GameplayTagsManager.cpp b/Engine/Source/Runtime/GameplayTags/Private/GameplayTagsManager.cpp
--- a/Engine/Source/Runtime/GameplayTags/Private/GameplayTagsManager.cpp
+++ b/Engine/Source/Runtime/GameplayTags/Private/GameplayTagsManager.cpp
@@ -699,8 +699,13 @@ bool UGameplayTagsManager::AddLeafTagToContainer(FGameplayTagContainer& TagConta
 	return true;
 }
 
+DECLARE_CYCLE_STAT(TEXT("AddParentTags"), STAT_AddParentTags, STATGROUP_GameplayTags);
+
+
 void UGameplayTagsManager::AddParentTags(FGameplayTagContainer& TagContainer, const FGameplayTag& GameplayTag) const
 {
+	SCOPE_CYCLE_COUNTER(STAT_AddParentTags);
+
 	const TSharedPtr<FGameplayTagNode>* GameplayTagNode = GameplayTagNodeMap.Find(GameplayTag);
 	if (GameplayTagNode)
 	{
diff --git a/Engine/Source/Runtime/Json/Public/Serialization/JsonReader.h b/Engine/Source/Runtime/Json/Public/Serialization/JsonReader.h
--- a/Engine/Source/Runtime/Json/Public/Serialization/JsonReader.h
+++ b/Engine/Source/Runtime/Json/Public/Serialization/JsonReader.h
@@ -586,6 +586,12 @@ private:
 
 		while (true)
 		{
+			if (Stream->AtEnd())
+			{
+				SetErrorMessage( TEXT("Number Token Abruptly Ended.") );
+				return false;
+			}
+
 			CharType Char;
 			if (UseFirstChar)
 			{
diff --git a/Engine/Source/Runtime/Launch/Resources/Version.h b/Engine/Source/Runtime/Launch/Resources/Version.h
--- a/Engine/Source/Runtime/Launch/Resources/Version.h
+++ b/Engine/Source/Runtime/Launch/Resources/Version.h
@@ -9,7 +9,7 @@
 #define ENGINE_IS_LICENSEE_VERSION 0
 
 #define BUILT_FROM_CHANGELIST 0
-#define BRANCH_NAME "UE4"
+#define BRANCH_NAME "UE4-Orion"
 
 #define ENGINE_IS_PROMOTED_BUILD (BUILT_FROM_CHANGELIST > 0)
 
diff --git a/Engine/Source/Runtime/MovieScene/Private/MovieSceneSection.cpp b/Engine/Source/Runtime/MovieScene/Private/MovieSceneSection.cpp
--- a/Engine/Source/Runtime/MovieScene/Private/MovieSceneSection.cpp
+++ b/Engine/Source/Runtime/MovieScene/Private/MovieSceneSection.cpp
@@ -10,9 +10,21 @@ UMovieSceneSection::UMovieSceneSection( const FObjectInitializer& ObjectInitiali
 	, EndTime(0.0f)
 	, RowIndex(0)
 	, bIsActive(true)
+	, bIsLocked(false)
 	, bIsInfinite(false)
 { }
 
+bool
+UMovieSceneSection::TryModify(bool bAlwaysMarkDirty)
+{
+	if (IsLocked())
+	{
+		return false;
+	}
+
+	return Modify(bAlwaysMarkDirty);
+}
+
 
 const UMovieSceneSection* UMovieSceneSection::OverlapsWithSections(const TArray<UMovieSceneSection*>& Sections, int32 TrackDelta, float TimeDelta) const
 {
@@ -75,25 +87,29 @@ UMovieSceneSection* UMovieSceneSection::SplitSection(float SplitTime)
 	}
 
 	SetFlags(RF_Transactional);
-	Modify();
 
-	float SectionEndTime = GetEndTime();
+	if (TryModify())
+	{
+		float SectionEndTime = GetEndTime();
 				
-	// Trim off the right
-	SetEndTime(SplitTime);
+		// Trim off the right
+		SetEndTime(SplitTime);
+
+		// Create a new section
+		UMovieSceneTrack* Track = CastChecked<UMovieSceneTrack>(GetOuter());
+		Track->Modify();
 
-	// Create a new section
-	UMovieSceneTrack* Track = CastChecked<UMovieSceneTrack>(GetOuter());
-	Track->Modify();
+		UMovieSceneSection* NewSection = DuplicateObject<UMovieSceneSection>(this, Track);
+		ensure(NewSection);
 
-	UMovieSceneSection* NewSection = DuplicateObject<UMovieSceneSection>(this, Track);
-	ensure(NewSection);
+		NewSection->SetStartTime(SplitTime);
+		NewSection->SetEndTime(SectionEndTime);
+		Track->AddSection(*NewSection);
 
-	NewSection->SetStartTime(SplitTime);
-	NewSection->SetEndTime(SectionEndTime);
-	Track->AddSection(*NewSection);
+		return NewSection;
+	}
 
-	return NewSection;
+	return nullptr;
 }
 
 
@@ -102,15 +118,16 @@ void UMovieSceneSection::TrimSection(float TrimTime, bool bTrimLeft)
 	if (IsTimeWithinSection(TrimTime))
 	{
 		SetFlags(RF_Transactional);
-		Modify();
-
-		if (bTrimLeft)
-		{
-			SetStartTime(TrimTime);
-		}
-		else
+		if (TryModify())
 		{
-			SetEndTime(TrimTime);
+			if (bTrimLeft)
+			{
+				SetStartTime(TrimTime);
+			}
+			else
+			{
+				SetEndTime(TrimTime);
+			}
 		}
 	}
 }
@@ -119,21 +136,25 @@ void UMovieSceneSection::AddKeyToCurve(FRichCurve& InCurve, float Time, float Va
 {
 	if (IsTimeWithinSection(Time))
 	{
-		Modify();
-		FKeyHandle ExistingKeyHandle = InCurve.FindKey(Time);
-			
-		FKeyHandle NewKeyHandle = InCurve.UpdateOrAddKey(Time, Value, bUnwindRotation);
-
-		if (!InCurve.IsKeyHandleValid(ExistingKeyHandle) && InCurve.IsKeyHandleValid(NewKeyHandle))
+		if (TryModify())
 		{
-			MovieSceneHelpers::SetKeyInterpolation(InCurve, NewKeyHandle, Interpolation);
+			FKeyHandle ExistingKeyHandle = InCurve.FindKey(Time);
+				
+			FKeyHandle NewKeyHandle = InCurve.UpdateOrAddKey(Time, Value, bUnwindRotation);
+
+			if (!InCurve.IsKeyHandleValid(ExistingKeyHandle) && InCurve.IsKeyHandleValid(NewKeyHandle))
+			{
+				MovieSceneHelpers::SetKeyInterpolation(InCurve, NewKeyHandle, Interpolation);
+			}
 		}
 	}
 }
 
 void UMovieSceneSection::SetCurveDefault(FRichCurve& InCurve, float Value)
 {
-	Modify();
-	InCurve.SetDefaultValue(Value);
+	if (TryModify())
+	{
+		InCurve.SetDefaultValue(Value);
+	}
 }
 
diff --git a/Engine/Source/Runtime/MovieScene/Public/MovieSceneSection.h b/Engine/Source/Runtime/MovieScene/Public/MovieSceneSection.h
--- a/Engine/Source/Runtime/MovieScene/Public/MovieSceneSection.h
+++ b/Engine/Source/Runtime/MovieScene/Public/MovieSceneSection.h
@@ -16,6 +16,13 @@ class UMovieSceneSection
 	GENERATED_UCLASS_BODY()
 public:
 
+	/*
+	 * Calls Modify if this section can be modified. ie. can't be modified if it's locked
+	 *
+	 * @return The return value of Modify()
+	 */
+	virtual MOVIESCENE_API bool TryModify( bool bAlwaysMarkDirty=true );
+
 	/**
 	 * @return The start time of the section
 	 */
@@ -47,9 +54,10 @@ public:
 	 */
 	void SetStartTime( float NewStartTime )
 	{ 
-		Modify();
-
-		StartTime = NewStartTime;
+		if (TryModify())
+		{
+			StartTime = NewStartTime;
+		}
 	}
 
 	/**
@@ -59,9 +67,10 @@ public:
 	 */
 	void SetEndTime( float NewEndTime )
 	{ 
-		Modify();
-
-		EndTime = NewEndTime;
+		if (TryModify())
+		{
+			EndTime = NewEndTime;
+		}
 	}
 	
 	/**
@@ -86,10 +95,11 @@ public:
 	{
 		check(NewRange.HasLowerBound() && NewRange.HasUpperBound());
 
-		Modify();
-
-		StartTime = NewRange.GetLowerBoundValue();
-		EndTime = NewRange.GetUpperBoundValue();
+		if (TryModify())
+		{
+			StartTime = NewRange.GetLowerBoundValue();
+			EndTime = NewRange.GetUpperBoundValue();
+		}
 	}
 
 	/**
@@ -111,10 +121,11 @@ public:
 	 */
 	virtual void MoveSection( float DeltaTime, TSet<FKeyHandle>& KeyHandles )
 	{
-		Modify();
-
-		StartTime += DeltaTime;
-		EndTime += DeltaTime;
+		if (TryModify())
+		{
+			StartTime += DeltaTime;
+			EndTime += DeltaTime;
+		}
 	}
 	
 	/**
@@ -126,10 +137,11 @@ public:
 	 */
 	virtual void DilateSection( float DilationFactor, float Origin, TSet<FKeyHandle>& KeyHandles )
 	{
-		Modify();
-
-		StartTime = (StartTime - Origin) * DilationFactor + Origin;
-		EndTime = (EndTime - Origin) * DilationFactor + Origin;
+		if (TryModify())
+		{
+			StartTime = (StartTime - Origin) * DilationFactor + Origin;
+			EndTime = (EndTime - Origin) * DilationFactor + Origin;
+		}
 	}
 	
 	/**
@@ -222,6 +234,10 @@ public:
 	void SetIsActive(bool bInIsActive) { bIsActive = bInIsActive; }
 	bool IsActive() const { return bIsActive; }
 
+	/** Whether or not this section is locked. */
+	void SetIsLocked(bool bInIsLocked) { bIsLocked = bInIsLocked; }
+	bool IsLocked() const { return bIsLocked; }
+
 	/** Whether or not this section is infinite. An infinite section will draw the entire width of the track. StartTime and EndTime will be ignored but not discarded. */
 	void SetIsInfinite(bool bInIsInfinite) { bIsInfinite = bInIsInfinite; }
 	bool IsInfinite() const { return bIsInfinite; }
@@ -244,6 +260,10 @@ private:
 	UPROPERTY(EditAnywhere, Category="Section")
 	uint32 bIsActive : 1;
 
+	/** Toggle whether this section is locked/unlocked */
+	UPROPERTY(EditAnywhere, Category="Section")
+	uint32 bIsLocked : 1;
+
 	/** Toggle to set this section to be infinite */
 	UPROPERTY(EditAnywhere, Category="Section")
 	uint32 bIsInfinite : 1;
diff --git a/Engine/Source/Runtime/MovieSceneCapture/Private/AutomatedLevelSequenceCapture.cpp b/Engine/Source/Runtime/MovieSceneCapture/Private/AutomatedLevelSequenceCapture.cpp
--- a/Engine/Source/Runtime/MovieSceneCapture/Private/AutomatedLevelSequenceCapture.cpp
+++ b/Engine/Source/Runtime/MovieSceneCapture/Private/AutomatedLevelSequenceCapture.cpp
@@ -16,7 +16,15 @@ UAutomatedLevelSequenceCapture::UAutomatedLevelSequenceCapture(const FObjectInit
 		checkf(false, TEXT("Automated level sequence captures can only be used in editor builds."));
 	}
 #else
-	bStageSequence = false;
+	bUseCustomStartFrame = false;
+	StartFrame = 0;
+	bUseCustomEndFrame = false;
+	EndFrame = 1;
+	WarmUpFrameCount = 0;
+	DelayBeforeWarmUp = 0.0f;
+
+	RemainingDelaySeconds = 0.0f;
+	RemainingWarmUpFrames = 0;
 #endif
 }
 
@@ -44,16 +52,30 @@ void UAutomatedLevelSequenceCapture::Initialize(TWeakPtr<FSceneViewport> InViewp
 			LevelSequenceActorId = FGuid();
 		}
 
-		float PrerollOverride;
-		if( FParse::Value( FCommandLine::Get(), TEXT( "-MoviePrerollSeconds=" ), PrerollOverride ) )
+		int32 StartFrameOverride;
+		if( FParse::Value( FCommandLine::Get(), TEXT( "-MovieStartFrame=" ), StartFrameOverride ) )
 		{
-			PrerollAmount = PrerollOverride;
+			bUseCustomStartFrame = true;
+			StartFrame = StartFrameOverride;
 		}
 
-		bool bStageSequenceOverride;
-		if( FParse::Bool( FCommandLine::Get(), TEXT( "-MoviePreStage=" ), bStageSequenceOverride ) )
+		int32 EndFrameOverride;
+		if( FParse::Value( FCommandLine::Get(), TEXT( "-MovieEndFrame=" ), EndFrameOverride ) )
 		{
-			bStageSequence = bStageSequenceOverride;
+			bUseCustomEndFrame = true;
+			EndFrame = EndFrameOverride;
+		}
+
+		int32 WarmUpFrameCountOverride;
+		if( FParse::Value( FCommandLine::Get(), TEXT( "-MovieWarmUpFrames=" ), WarmUpFrameCountOverride ) )
+		{
+			WarmUpFrameCount = WarmUpFrameCountOverride;
+		}
+
+		float DelayBeforeWarmUpOverride;
+		if( FParse::Value( FCommandLine::Get(), TEXT( "-MovieDelayBeforeWarmUp=" ), DelayBeforeWarmUpOverride ) )
+		{
+			DelayBeforeWarmUp = DelayBeforeWarmUpOverride;
 		}
 	}
 
@@ -100,8 +122,9 @@ void UAutomatedLevelSequenceCapture::Initialize(TWeakPtr<FSceneViewport> InViewp
 			OnPlayerUpdatedBinding = Actor->SequencePlayer->OnSequenceUpdated().AddUObject( this, &UAutomatedLevelSequenceCapture::SequenceUpdated );
 		}
 	}
-	
-	PrerollTime = 0.f;
+
+	CaptureState = ELevelSequenceCaptureState::DelayBeforeWarmUp;
+	RemainingDelaySeconds = FMath::Max( 0.0f, DelayBeforeWarmUp );
 	CaptureStrategy = MakeShareable(new FFixedTimeStepCaptureStrategy(Settings.FrameRate));
 
 	Super::Initialize(InViewport);
@@ -119,7 +142,6 @@ void UAutomatedLevelSequenceCapture::SetupFrameRange()
 			UMovieScene* MovieScene = LevelSequence->GetMovieScene();
 			if( MovieScene != nullptr )
 			{
-				// We always add 1 to the number of frames we want to capture, because we want to capture both the start and end frames (which if the play range is 0, would still yield a single frame)
 				const int32 SequenceStartFrame = FMath::RoundToInt( MovieScene->GetPlaybackRange().GetLowerBoundValue() * Settings.FrameRate );
 				const int32 SequenceEndFrame = FMath::Max( SequenceStartFrame, FMath::RoundToInt( MovieScene->GetPlaybackRange().GetUpperBoundValue() * Settings.FrameRate ) );
 
@@ -127,21 +149,32 @@ void UAutomatedLevelSequenceCapture::SetupFrameRange()
 				int32 PlaybackStartFrame = SequenceStartFrame;
 				int32 PlaybackEndFrame = SequenceEndFrame;
 
-				if( Settings.bUseCustomStartFrame )
+				if( bUseCustomStartFrame )
 				{
-					PlaybackStartFrame = Settings.bUseRelativeFrameNumbers ? ( SequenceStartFrame + Settings.StartFrame ) : Settings.StartFrame;
+					PlaybackStartFrame = Settings.bUseRelativeFrameNumbers ? ( SequenceStartFrame + StartFrame ) : StartFrame;
 				}
 
-				if( Settings.bUseCustomEndFrame )
+				if( !Settings.bUseRelativeFrameNumbers )
 				{
-					PlaybackEndFrame = FMath::Max( PlaybackStartFrame, Settings.bUseRelativeFrameNumbers ? ( SequenceEndFrame + Settings.EndFrame ) : Settings.EndFrame );
+					// NOTE: The frame number will be an offset from the first frame that we start capturing on, not the frame
+					// that we start playback at (in the case of WarmUpFrameCount being non-zero).  So we'll cache out frame
+					// number offset before adjusting for the warm up frames.
+					this->FrameNumberOffset = PlaybackStartFrame;
 				}
 
+				if( bUseCustomEndFrame )
+				{
+					PlaybackEndFrame = FMath::Max( PlaybackStartFrame, Settings.bUseRelativeFrameNumbers ? ( SequenceEndFrame + EndFrame ) : EndFrame );
+				}
+
+				// We always add 1 to the number of frames we want to capture, because we want to capture both the start and end frames (which if the play range is 0, would still yield a single frame)
 				this->FrameCount = ( PlaybackEndFrame - PlaybackStartFrame ) + 1;
 
-				if( !Settings.bUseRelativeFrameNumbers )
+				RemainingWarmUpFrames = FMath::Max( WarmUpFrameCount, 0 );
+				if( RemainingWarmUpFrames > 0 )
 				{
-					this->FrameNumberOffset = PlaybackStartFrame;
+					// We were asked to playback additional frames before we start capturing
+					PlaybackStartFrame -= RemainingWarmUpFrames;
 				}
 
 				// Override the movie scene's playback range
@@ -182,37 +215,64 @@ void UAutomatedLevelSequenceCapture::Tick(float DeltaSeconds)
 
 	if (Actor && Actor->SequencePlayer)
 	{
-		// First off, check if we need to stage the sequence
-		if (bStageSequence)
+		// First off we'll stage the sequence.  This just means we'll play the first frame of the animation
+		// and then pause it immediately
+		if( CaptureState == ELevelSequenceCaptureState::Staging )
 		{
 			Actor->SequencePlayer->Play();
 			Actor->SequencePlayer->Pause();
-			bStageSequence = false;
+
+			CaptureState = ELevelSequenceCaptureState::DelayBeforeWarmUp;
 		}
 
-		// Then handle any preroll
-		if (PrerollTime.IsSet())
+		// Then we'll just wait a little bit.  We'll delay the specified number of seconds before capturing to allow any
+		// textures to stream in or post processing effects to settle.
+		if( CaptureState == ELevelSequenceCaptureState::DelayBeforeWarmUp )
 		{
-			float& Value = PrerollTime.GetValue();
-			Value += DeltaSeconds;
-			if (Value >= PrerollAmount)
+			RemainingDelaySeconds -= DeltaSeconds;
+			if( RemainingDelaySeconds <= 0.0f )
 			{
-				Actor->SequencePlayer->Play();
+				RemainingDelaySeconds = 0.0f;
 
-				StartCapture();
+				// Start warming up.  Even if we're not capturing yet, this will make sure we're rendering at a
+				// fixed frame rate.
+				StartWarmup();
 
-				PrerollTime.Reset();
+				// Wait a frame to go by after we've set the fixed time step, so that the animation starts
+				// playback at a consistent time
+				CaptureState = ELevelSequenceCaptureState::ReadyToWarmUp;
 			}
 		}
+		else if( CaptureState == ELevelSequenceCaptureState::ReadyToWarmUp )
+		{
+			Actor->SequencePlayer->Play();
+			CaptureState = ELevelSequenceCaptureState::WarmingUp;
+		}
 
-		if( !PrerollTime.IsSet() )
+
+		if( CaptureState == ELevelSequenceCaptureState::WarmingUp )
 		{
-			if( bAnyFramesToCapture && OutstandingFrameCount == 0 )
+			// Count down our warm up frames
+			if( RemainingWarmUpFrames == 0 )
 			{
-				// If we hit this, then we've rendered out the last frame and can stop!
-				StopCapture();
+				CaptureState = ELevelSequenceCaptureState::FinishedWarmUp;
+
+				// It's time to start capturing!
+				StartCapture();
+			}
+			else
+			{
+				// Not ready to capture just yet
+				--RemainingWarmUpFrames;
 			}
 		}
+
+
+		if( bAnyFramesToCapture && OutstandingFrameCount == 0 )
+		{
+			// If we hit this, then we've rendered out the last frame and can stop!
+			StopCapture();
+		}
 	}
 }
 
diff --git a/Engine/Source/Runtime/MovieSceneCapture/Private/LevelCapture.cpp b/Engine/Source/Runtime/MovieSceneCapture/Private/LevelCapture.cpp
--- a/Engine/Source/Runtime/MovieSceneCapture/Private/LevelCapture.cpp
+++ b/Engine/Source/Runtime/MovieSceneCapture/Private/LevelCapture.cpp
@@ -30,15 +30,12 @@ void ULevelCapture::Tick(float DeltaSeconds)
 	AActor* PrerequisiteActorPtr = PrerequisiteActor.Get();
 	if (!PrerequisiteActorId.IsValid() || (PrerequisiteActorPtr && PrerequisiteActorPtr->HasActorBegunPlay()))
 	{
-		if( !Settings.bUseCustomStartFrame || GetMetrics().Frame >= Settings.StartFrame )
+		if (!bCapturing)
 		{
-			if (!bCapturing)
-			{
-				StartCapture();
-			}
-
-			PrepareForScreenshot();
+			StartCapture();
 		}
+
+		PrepareForScreenshot();
 	}
 }
 
diff --git a/Engine/Source/Runtime/MovieSceneCapture/Private/MovieSceneCapture.cpp b/Engine/Source/Runtime/MovieSceneCapture/Private/MovieSceneCapture.cpp
--- a/Engine/Source/Runtime/MovieSceneCapture/Private/MovieSceneCapture.cpp
+++ b/Engine/Source/Runtime/MovieSceneCapture/Private/MovieSceneCapture.cpp
@@ -31,12 +31,8 @@ FMovieSceneCaptureSettings::FMovieSceneCaptureSettings()
 	bCreateTemporaryCopiesOfLevels = false;
 	bUseRelativeFrameNumbers = false;
 	GameModeOverride = nullptr;
-	OutputFormat = NSLOCTEXT("MovieCapture", "DefaultFormat", "{world}_{width}x{height}").ToString();
+	OutputFormat = NSLOCTEXT("MovieCapture", "DefaultFormat", "{world}_{frame}").ToString();
 	FrameRate = 24;
-	bUseCustomStartFrame = false;
-	StartFrame = 0;
-	bUseCustomEndFrame = false;
-	EndFrame = 1;
 	CaptureType = EMovieCaptureType::AVI;
 	bUseCompression = true;
 	CompressionQuality = 1.f;
@@ -89,6 +85,12 @@ void UMovieSceneCapture::Initialize(TWeakPtr<FSceneViewport> InSceneViewport)
 			Settings.OutputFormat = OutputNameOverride;
 		}
 
+		bool bOverrideOverwriteExisting;
+		if( FParse::Bool( FCommandLine::Get(), TEXT( "-MovieOverwriteExisting=" ), bOverrideOverwriteExisting ) )
+		{
+			Settings.bOverwriteExisting = bOverrideOverwriteExisting;
+		}
+
 		bool bOverrideRelativeFrameNumbers;
 		if( FParse::Bool( FCommandLine::Get(), TEXT( "-MovieRelativeFrames=" ), bOverrideRelativeFrameNumbers ) )
 		{
@@ -122,20 +124,6 @@ void UMovieSceneCapture::Initialize(TWeakPtr<FSceneViewport> InSceneViewport)
 			}
 		}
 
-		int32 StartFrameOverride;
-		if( FParse::Value( FCommandLine::Get(), TEXT( "-MovieStartFrame=" ), StartFrameOverride ) )
-		{
-			Settings.bUseCustomStartFrame = true;
-			Settings.StartFrame = StartFrameOverride;
-		}
-
-		int32 EndFrameOverride;
-		if( FParse::Value( FCommandLine::Get(), TEXT( "-MovieEndFrame=" ), EndFrameOverride ) )
-		{
-			Settings.bUseCustomEndFrame = true;
-			Settings.EndFrame = EndFrameOverride;
-		}
-
 		int32 FrameRateOverride;
 		if( FParse::Value( FCommandLine::Get(), TEXT( "-MovieFrameRate=" ), FrameRateOverride ) )
 		{
@@ -211,6 +199,16 @@ void UMovieSceneCapture::Tick(float DeltaSeconds)
 	LastFrameDelta = DeltaSeconds;
 }
 
+void UMovieSceneCapture::StartWarmup()
+{
+	check( !bCapturing );
+	if( !CaptureStrategy.IsValid() )
+	{
+		CaptureStrategy = MakeShareable( new FRealTimeCaptureStrategy( Settings.FrameRate ) );
+	}
+	CaptureStrategy->OnWarmup();
+}
+
 void UMovieSceneCapture::StartCapture()
 {
 	auto Viewport = SceneViewport.Pin();
@@ -226,6 +224,7 @@ void UMovieSceneCapture::StartCapture()
 		CaptureStrategy = MakeShareable(new FRealTimeCaptureStrategy(Settings.FrameRate));
 	}
 
+	check( CaptureStrategy.IsValid() );
 	CaptureStrategy->OnStart();
 
 	CachedMetrics.ElapsedSeconds = 0;
@@ -473,6 +472,12 @@ FFixedTimeStepCaptureStrategy::FFixedTimeStepCaptureStrategy(uint32 InTargetFPS)
 {
 }
 
+void FFixedTimeStepCaptureStrategy::OnWarmup()
+{
+	FApp::SetFixedDeltaTime(1.0 / TargetFPS);
+	FApp::SetUseFixedTimeStep(true);
+}
+
 void FFixedTimeStepCaptureStrategy::OnStart()
 {
 	FApp::SetFixedDeltaTime(1.0 / TargetFPS);
@@ -503,6 +508,10 @@ FRealTimeCaptureStrategy::FRealTimeCaptureStrategy(uint32 InTargetFPS)
 {
 }
 
+void FRealTimeCaptureStrategy::OnWarmup()
+{
+}
+
 void FRealTimeCaptureStrategy::OnStart()
 {
 }
diff --git a/Engine/Source/Runtime/MovieSceneCapture/Public/AutomatedLevelSequenceCapture.h b/Engine/Source/Runtime/MovieSceneCapture/Public/AutomatedLevelSequenceCapture.h
--- a/Engine/Source/Runtime/MovieSceneCapture/Public/AutomatedLevelSequenceCapture.h
+++ b/Engine/Source/Runtime/MovieSceneCapture/Public/AutomatedLevelSequenceCapture.h
@@ -24,13 +24,30 @@ public:
 	/** Set the level sequence actor that we are to record */
 	void SetLevelSequenceActor(ALevelSequenceActor* InActor);
 
-	/** The amount of time to wait before playback and capture start. Useful for allowing Post Processing effects to settle down before capturing the animation. */
-	UPROPERTY(config, EditAnywhere, Category=CaptureSettings, AdvancedDisplay, meta=(Units=Seconds, ClampMin=0))
-	float PrerollAmount;
+	/** When enabled, the StartFrame setting will override the default starting frame number */
+	UPROPERTY(config, EditAnywhere, Category=CaptureSettings, AdvancedDisplay)
+	bool bUseCustomStartFrame;
+
+	/** Frame number to start capturing.  The frame number range depends on whether the bUseRelativeFrameNumbers option is enabled. */
+	UPROPERTY(config, EditAnywhere, Category=CaptureSettings, AdvancedDisplay, meta=(EditCondition="bUseCustomStartFrame"))
+	int32 StartFrame;
 
-	/** Whether to stage the sequence before capturing it. This will set the sequence to its first frame for Preroll Amount seconds to allow Post Processing effects to settle. */
+	/** When enabled, the EndFrame setting will override the default ending frame number */
 	UPROPERTY(config, EditAnywhere, Category=CaptureSettings, AdvancedDisplay)
-	bool bStageSequence;
+	bool bUseCustomEndFrame;
+
+	/** Frame number to end capturing.  The frame number range depends on whether the bUseRelativeFrameNumbers option is enabled. */
+	UPROPERTY(config, EditAnywhere, Category=CaptureSettings, AdvancedDisplay, meta=(EditCondition="bUseCustomEndFrame"))
+	int32 EndFrame;
+
+	/** The number of extra frames to play before the sequence's start frame, to "warm up" the animation.  This is useful if your
+	    animation contains particles or other runtime effects that are spawned into the scene earlier than your capture start frame */
+	UPROPERTY(config, EditAnywhere, Category=CaptureSettings, AdvancedDisplay)
+	int32 WarmUpFrameCount;
+
+	/** The number of seconds to wait (in real-time) before we start playing back the warm up frames.  Useful for allowing post processing effects to settle down before capturing the animation. */
+	UPROPERTY(config, EditAnywhere, Category=CaptureSettings, AdvancedDisplay, meta=(Units=Seconds, ClampMin=0))
+	float DelayBeforeWarmUp;
 
 public:
 	// UMovieSceneCapture interface
@@ -65,9 +82,22 @@ private:
 	UPROPERTY()
 	FGuid LevelSequenceActorId;
 
-	/** Transient amount of time that has passed since the level was started. When >= Settings.Preroll, playback will start  */
-	TOptional<float> PrerollTime;
-
+	/** Which state we're in right now */
+	enum class ELevelSequenceCaptureState
+	{
+		Staging,
+		DelayBeforeWarmUp,
+		ReadyToWarmUp,
+		WarmingUp,
+		FinishedWarmUp
+	} CaptureState;
+
+	/** Time left to wait before capturing */
+	float RemainingDelaySeconds;
+
+	/** The number of warm up frames left before we actually start saving out images */
+	int32 RemainingWarmUpFrames;
+	
 #endif
 };
 
diff --git a/Engine/Source/Runtime/MovieSceneCapture/Public/MovieSceneCapture.h b/Engine/Source/Runtime/MovieSceneCapture/Public/MovieSceneCapture.h
--- a/Engine/Source/Runtime/MovieSceneCapture/Public/MovieSceneCapture.h
+++ b/Engine/Source/Runtime/MovieSceneCapture/Public/MovieSceneCapture.h
@@ -15,6 +15,7 @@ struct ICaptureStrategy
 {
 	virtual ~ICaptureStrategy(){}
 	
+	virtual void OnWarmup() = 0;
 	virtual void OnStart() = 0;
 	virtual void OnStop() = 0;
 	virtual void OnPresent(double CurrentTimeSeconds, uint32 FrameIndex) = 0;
@@ -82,6 +83,10 @@ public:
 	/** Access this object's cached metrics */
 	const FCachedMetrics& GetMetrics() const { return CachedMetrics; }
 
+	/** Starts warming up.  May be optionally called before StartCapture().  This can be used to start rendering frames early, before
+	    any files are written out */
+	void StartWarmup();
+
 	/** Initialize the capture */
 	void StartCapture();
 
@@ -175,6 +180,7 @@ struct MOVIESCENECAPTURE_API FFixedTimeStepCaptureStrategy : ICaptureStrategy
 {
 	FFixedTimeStepCaptureStrategy(uint32 InTargetFPS);
 
+	virtual void OnWarmup() override;
 	virtual void OnStart() override;
 	virtual void OnStop() override;
 	virtual void OnPresent(double CurrentTimeSeconds, uint32 FrameIndex) override;	
@@ -190,6 +196,7 @@ struct MOVIESCENECAPTURE_API FRealTimeCaptureStrategy : ICaptureStrategy
 {
 	FRealTimeCaptureStrategy(uint32 InTargetFPS);
 
+	virtual void OnWarmup() override;
 	virtual void OnStart() override;
 	virtual void OnStop() override;
 	virtual void OnPresent(double CurrentTimeSeconds, uint32 FrameIndex) override;
diff --git a/Engine/Source/Runtime/MovieSceneCapture/Public/MovieSceneCaptureSettings.h b/Engine/Source/Runtime/MovieSceneCapture/Public/MovieSceneCaptureSettings.h
--- a/Engine/Source/Runtime/MovieSceneCapture/Public/MovieSceneCaptureSettings.h
+++ b/Engine/Source/Runtime/MovieSceneCapture/Public/MovieSceneCaptureSettings.h
@@ -38,7 +38,7 @@ struct MOVIESCENECAPTURE_API FMovieSceneCaptureSettings
 	GENERATED_BODY()
 
 	/** The directory to output the captured file(s) in */
-	UPROPERTY(config, EditAnywhere, Category=General, AdvancedDisplay, meta=(RelativePath))
+	UPROPERTY(config, EditAnywhere, Category=General, meta=(RelativePath))
 	FDirectoryPath OutputDirectory;
 
 	/** Whether to save temporary copies of all of the levels before capturing the movie.  This allows you to record movies of temporary work, or work that isn't yet saved, but it will take much longer for capturing to begin. */
@@ -72,22 +72,6 @@ struct MOVIESCENECAPTURE_API FMovieSceneCaptureSettings
 	UPROPERTY(config, EditAnywhere, Category=CaptureSettings)
 	int32 FrameRate;
 
-	/** When enabled, the StartFrame setting will override the default starting frame number */
-	UPROPERTY(config, EditAnywhere, Category=CaptureSettings, AdvancedDisplay )
-	bool bUseCustomStartFrame;
-
-	/** Frame number to start capturing.  The frame number range depends on whether the bUseRelativeFrameNumbers option is enabled. */
-	UPROPERTY(config, EditAnywhere, Category=CaptureSettings, AdvancedDisplay, meta=(EditCondition="bUseCustomStartFrame") )
-	int32 StartFrame;
-
-	/** When enabled, the EndFrame setting will override the default ending frame number */
-	UPROPERTY(config, EditAnywhere, Category=CaptureSettings, AdvancedDisplay )
-	bool bUseCustomEndFrame;
-
-	/** Frame number to end capturing.  The frame number range depends on whether the bUseRelativeFrameNumbers option is enabled. */
-	UPROPERTY(config, EditAnywhere, Category=CaptureSettings, AdvancedDisplay, meta=(EditCondition="bUseCustomEndFrame") )
-	int32 EndFrame;
-
 	/** The resolution at which to capture */
 	UPROPERTY(config, EditAnywhere, Category=CaptureSettings, meta=(ShowOnlyInnerProperties))
 	FCaptureResolution Resolution;
@@ -108,8 +92,8 @@ struct MOVIESCENECAPTURE_API FMovieSceneCaptureSettings
 	UPROPERTY(config, EditAnywhere, Category=CaptureSettings, AdvancedDisplay)
 	FString Codec;
 	
-	/** Whether to enable texture streaming whilst capturing */
-	UPROPERTY(config, EditAnywhere, Category=CaptureSettings)
+	/** Whether to texture streaming should be enabled while capturing.  Turning off texture streaming may cause much more memory to be used, but also reduces the chance of blurry textures in your captured video. */
+	UPROPERTY(config, EditAnywhere, Category=CaptureSettings, AdvancedDisplay)
 	bool bEnableTextureStreaming;
 
 	/** Whether to enable cinematic mode whilst capturing */
diff --git a/Engine/Source/Runtime/MovieSceneTracks/Private/Sections/MovieScene3DAttachSection.cpp b/Engine/Source/Runtime/MovieSceneTracks/Private/Sections/MovieScene3DAttachSection.cpp
--- a/Engine/Source/Runtime/MovieSceneTracks/Private/Sections/MovieScene3DAttachSection.cpp
+++ b/Engine/Source/Runtime/MovieSceneTracks/Private/Sections/MovieScene3DAttachSection.cpp
@@ -66,7 +66,9 @@ void UMovieScene3DAttachSection::Eval( USceneComponent* SceneComponent, float Po
 
 void UMovieScene3DAttachSection::AddAttach( float Time, float SequenceEndTime, const FGuid& InAttachId )
 {
-	Modify();
-	ConstraintId = InAttachId;
+	if (TryModify())
+	{
+		ConstraintId = InAttachId;
+	}
 }
 
diff --git a/Engine/Source/Runtime/MovieSceneTracks/Private/Sections/MovieScene3DConstraintSection.cpp b/Engine/Source/Runtime/MovieSceneTracks/Private/Sections/MovieScene3DConstraintSection.cpp
--- a/Engine/Source/Runtime/MovieSceneTracks/Private/Sections/MovieScene3DConstraintSection.cpp
+++ b/Engine/Source/Runtime/MovieSceneTracks/Private/Sections/MovieScene3DConstraintSection.cpp
@@ -11,9 +11,10 @@ UMovieScene3DConstraintSection::UMovieScene3DConstraintSection( const FObjectIni
 
 void UMovieScene3DConstraintSection::SetConstraintId(const FGuid& InConstraintId)
 {
-	Modify();
-
-	ConstraintId = InConstraintId;
+	if (TryModify())
+	{
+		ConstraintId = InConstraintId;
+	}
 }
 
 
diff --git a/Engine/Source/Runtime/MovieSceneTracks/Private/Sections/MovieScene3DPathSection.cpp b/Engine/Source/Runtime/MovieSceneTracks/Private/Sections/MovieScene3DPathSection.cpp
--- a/Engine/Source/Runtime/MovieSceneTracks/Private/Sections/MovieScene3DPathSection.cpp
+++ b/Engine/Source/Runtime/MovieSceneTracks/Private/Sections/MovieScene3DPathSection.cpp
@@ -144,9 +144,11 @@ void UMovieScene3DPathSection::GetKeyHandles(TSet<FKeyHandle>& KeyHandles) const
 
 void UMovieScene3DPathSection::AddPath( float Time, float SequenceEndTime, const FGuid& InPathId )
 {
-	Modify();
-	ConstraintId = InPathId;
+	if (TryModify())
+	{
+		ConstraintId = InPathId;
 
-	TimingCurve.UpdateOrAddKey(Time, 0);
-	TimingCurve.UpdateOrAddKey(SequenceEndTime, 1);
+		TimingCurve.UpdateOrAddKey(Time, 0);
+		TimingCurve.UpdateOrAddKey(SequenceEndTime, 1);
+	}		
 }
diff --git a/Engine/Source/Runtime/MovieSceneTracks/Private/Sections/MovieSceneBoolSection.cpp b/Engine/Source/Runtime/MovieSceneTracks/Private/Sections/MovieSceneBoolSection.cpp
--- a/Engine/Source/Runtime/MovieSceneTracks/Private/Sections/MovieSceneBoolSection.cpp
+++ b/Engine/Source/Runtime/MovieSceneTracks/Private/Sections/MovieSceneBoolSection.cpp
@@ -52,15 +52,19 @@ void UMovieSceneBoolSection::GetKeyHandles(TSet<FKeyHandle>& KeyHandles) const
 
 void UMovieSceneBoolSection::AddKey( float Time, const bool& Value, EMovieSceneKeyInterpolation KeyInterpolation )
 {
-	Modify();
-	BoolCurve.UpdateOrAddKey(Time, Value ? 1 : 0);
+	if (TryModify())
+	{
+		BoolCurve.UpdateOrAddKey(Time, Value ? 1 : 0);
+	}
 }
 
 
 void UMovieSceneBoolSection::SetDefault( const bool& Value )
 {
-	Modify();
-	BoolCurve.SetDefaultValue(Value ? 1 : 0);
+	if (TryModify())
+	{
+		BoolCurve.SetDefaultValue(Value ? 1 : 0);
+	}
 }
 
 
diff --git a/Engine/Source/Runtime/MovieSceneTracks/Private/Sections/MovieSceneByteSection.cpp b/Engine/Source/Runtime/MovieSceneTracks/Private/Sections/MovieSceneByteSection.cpp
--- a/Engine/Source/Runtime/MovieSceneTracks/Private/Sections/MovieSceneByteSection.cpp
+++ b/Engine/Source/Runtime/MovieSceneTracks/Private/Sections/MovieSceneByteSection.cpp
@@ -44,8 +44,10 @@ void UMovieSceneByteSection::GetKeyHandles(TSet<FKeyHandle>& KeyHandles) const
 
 void UMovieSceneByteSection::AddKey( float Time, const uint8& Value, EMovieSceneKeyInterpolation KeyInterpolation )
 {
-	Modify();
-	ByteCurve.UpdateOrAddKey(Time, Value);
+	if (TryModify())
+	{
+		ByteCurve.UpdateOrAddKey(Time, Value);
+	}
 }
 
 
@@ -63,6 +65,8 @@ bool UMovieSceneByteSection::HasKeys( const uint8& Value ) const
 
 void UMovieSceneByteSection::SetDefault( const uint8& Value )
 {
-	Modify();
-	ByteCurve.SetDefaultValue(Value);
+	if (TryModify())
+	{
+		ByteCurve.SetDefaultValue(Value);
+	}
 }
diff --git a/Engine/Source/Runtime/MovieSceneTracks/Private/Sections/MovieSceneEventSection.cpp b/Engine/Source/Runtime/MovieSceneTracks/Private/Sections/MovieSceneEventSection.cpp
--- a/Engine/Source/Runtime/MovieSceneTracks/Private/Sections/MovieSceneEventSection.cpp
+++ b/Engine/Source/Runtime/MovieSceneTracks/Private/Sections/MovieSceneEventSection.cpp
@@ -18,8 +18,10 @@ UMovieSceneEventSection::UMovieSceneEventSection()
 
 void UMovieSceneEventSection::AddKey(float Time, const FName& EventName, FKeyParams KeyParams)
 {
-	Modify();
-	Events.UpdateOrAddKey(Time, EventName);
+	if (TryModify())
+	{
+		Events.UpdateOrAddKey(Time, EventName);
+	}
 }
 
 
diff --git a/Engine/Source/Runtime/MovieSceneTracks/Private/Sections/MovieSceneShotSection.cpp b/Engine/Source/Runtime/MovieSceneTracks/Private/Sections/MovieSceneShotSection.cpp
--- a/Engine/Source/Runtime/MovieSceneTracks/Private/Sections/MovieSceneShotSection.cpp
+++ b/Engine/Source/Runtime/MovieSceneTracks/Private/Sections/MovieSceneShotSection.cpp
@@ -9,8 +9,9 @@
 
 void UMovieSceneShotSection::SetShotNameAndNumber(const FText& InDisplayName, int32 InShotNumber)
 {
-	Modify();
-
-	DisplayName = InDisplayName;
-	ShotNumber = InShotNumber;
+	if (TryModify())
+	{
+		DisplayName = InDisplayName;
+		ShotNumber = InShotNumber;
+	}
 }
diff --git a/Engine/Source/Runtime/MovieSceneTracks/Private/TrackInstances/MovieSceneSkeletalAnimationTrackInstance.cpp b/Engine/Source/Runtime/MovieSceneTracks/Private/TrackInstances/MovieSceneSkeletalAnimationTrackInstance.cpp
--- a/Engine/Source/Runtime/MovieSceneTracks/Private/TrackInstances/MovieSceneSkeletalAnimationTrackInstance.cpp
+++ b/Engine/Source/Runtime/MovieSceneTracks/Private/TrackInstances/MovieSceneSkeletalAnimationTrackInstance.cpp
@@ -68,9 +68,6 @@ void FMovieSceneSkeletalAnimationTrackInstance::Update( float Position, float La
 				AnimPosition = FMath::Fmod(AnimPosition, SeqLength);
 				AnimPosition += AnimSection->GetStartOffset();
 
-				// Clamp to end of section
-				AnimPosition = FMath::Clamp(AnimPosition, AnimPosition, AnimSection->GetEndTime());
-
 				if (AnimSection->GetReverse())
 				{
 					AnimPosition = (SeqLength - (AnimPosition - AnimSection->GetStartOffset())) + AnimSection->GetStartOffset();
diff --git a/Engine/Source/Runtime/MovieSceneTracks/Private/TrackInstances/MovieSceneSlomoTrackInstance.cpp b/Engine/Source/Runtime/MovieSceneTracks/Private/TrackInstances/MovieSceneSlomoTrackInstance.cpp
--- a/Engine/Source/Runtime/MovieSceneTracks/Private/TrackInstances/MovieSceneSlomoTrackInstance.cpp
+++ b/Engine/Source/Runtime/MovieSceneTracks/Private/TrackInstances/MovieSceneSlomoTrackInstance.cpp
@@ -10,8 +10,33 @@
 
 FMovieSceneSlomoTrackInstance::FMovieSceneSlomoTrackInstance(UMovieSceneSlomoTrack& InSlomoTrack)
 	: SlomoTrack(&InSlomoTrack)
+	, InitMatineeTimeDilation(1.0f)
 { }
 
+	
+void FMovieSceneSlomoTrackInstance::RestoreState(const TArray<UObject*>& RuntimeObjects, IMovieScenePlayer& Player, FMovieSceneSequenceInstance& SequenceInstance)
+{
+	AWorldSettings* WorldSettings = GWorld->GetWorldSettings();
+
+	if (WorldSettings == nullptr)
+	{
+		return;
+	}
+
+	WorldSettings->MatineeTimeDilation = InitMatineeTimeDilation;
+}
+
+void FMovieSceneSlomoTrackInstance::SaveState(const TArray<UObject*>& RuntimeObjects, IMovieScenePlayer& Player, FMovieSceneSequenceInstance& SequenceInstance)
+{
+	AWorldSettings* WorldSettings = GWorld->GetWorldSettings();
+
+	if (WorldSettings == nullptr)
+	{
+		return;
+	}
+
+	InitMatineeTimeDilation = WorldSettings->MatineeTimeDilation;
+}
 
 /* IMovieSceneTrackInstance interface
  *****************************************************************************/
diff --git a/Engine/Source/Runtime/MovieSceneTracks/Public/TrackInstances/MovieSceneSlomoTrackInstance.h b/Engine/Source/Runtime/MovieSceneTracks/Public/TrackInstances/MovieSceneSlomoTrackInstance.h
--- a/Engine/Source/Runtime/MovieSceneTracks/Public/TrackInstances/MovieSceneSlomoTrackInstance.h
+++ b/Engine/Source/Runtime/MovieSceneTracks/Public/TrackInstances/MovieSceneSlomoTrackInstance.h
@@ -32,8 +32,8 @@ public:
 
 	virtual void ClearInstance(IMovieScenePlayer& Player, FMovieSceneSequenceInstance& SequenceInstance) override { }
 	virtual void RefreshInstance(const TArray<UObject*>& RuntimeObjects, IMovieScenePlayer& Player, FMovieSceneSequenceInstance& SequenceInstance) override { }
-	virtual void RestoreState(const TArray<UObject*>& RuntimeObjects, IMovieScenePlayer& Player, FMovieSceneSequenceInstance& SequenceInstance) override { }
-	virtual void SaveState(const TArray<UObject*>& RuntimeObjects, IMovieScenePlayer& Player, FMovieSceneSequenceInstance& SequenceInstance) override { }
+	virtual void RestoreState(const TArray<UObject*>& RuntimeObjects, IMovieScenePlayer& Player, FMovieSceneSequenceInstance& SequenceInstance) override;
+	virtual void SaveState(const TArray<UObject*>& RuntimeObjects, IMovieScenePlayer& Player, FMovieSceneSequenceInstance& SequenceInstance) override;
 	virtual void Update(float Position, float LastPosition, const TArray<UObject*>& RuntimeObjects, IMovieScenePlayer& Player, FMovieSceneSequenceInstance& SequenceInstance, EMovieSceneUpdatePass UpdatePass) override;
 
 protected:
@@ -48,4 +48,6 @@ private:
 
 	/** Track that is being instanced */
 	UMovieSceneSlomoTrack* SlomoTrack;
+
+	float InitMatineeTimeDilation;
 };
diff --git a/Engine/Source/Runtime/Online/Android/OnlineSubsystemGooglePlay/Private/OnlineAsyncTaskGooglePlayShowLoginUI.cpp b/Engine/Source/Runtime/Online/Android/OnlineSubsystemGooglePlay/Private/OnlineAsyncTaskGooglePlayShowLoginUI.cpp
--- a/Engine/Source/Runtime/Online/Android/OnlineSubsystemGooglePlay/Private/OnlineAsyncTaskGooglePlayShowLoginUI.cpp
+++ b/Engine/Source/Runtime/Online/Android/OnlineSubsystemGooglePlay/Private/OnlineAsyncTaskGooglePlayShowLoginUI.cpp
@@ -6,7 +6,7 @@
 FOnlineAsyncTaskGooglePlayShowLoginUI::FOnlineAsyncTaskGooglePlayShowLoginUI(
 	FOnlineSubsystemGooglePlay* InSubsystem,
 	int InPlayerId,
-	const IOnlineExternalUI::FOnLoginUIClosedDelegate& InDelegate)
+	const FOnLoginUIClosedDelegate& InDelegate)
 	: FOnlineAsyncTaskGooglePlayAuthAction(InSubsystem)
 	, PlayerId(InPlayerId)
 	, Delegate(InDelegate)
diff --git a/Engine/Source/Runtime/Online/Android/OnlineSubsystemGooglePlay/Private/OnlineAsyncTaskGooglePlayShowLoginUI.h b/Engine/Source/Runtime/Online/Android/OnlineSubsystemGooglePlay/Private/OnlineAsyncTaskGooglePlayShowLoginUI.h
--- a/Engine/Source/Runtime/Online/Android/OnlineSubsystemGooglePlay/Private/OnlineAsyncTaskGooglePlayShowLoginUI.h
+++ b/Engine/Source/Runtime/Online/Android/OnlineSubsystemGooglePlay/Private/OnlineAsyncTaskGooglePlayShowLoginUI.h
@@ -20,7 +20,7 @@ public:
 	 * @param InSubsystem a pointer to the owning subsystem
 	 * @param InPlayerId index of the player who's logging in
 	 */
-	FOnlineAsyncTaskGooglePlayShowLoginUI(FOnlineSubsystemGooglePlay* InSubsystem, int InPlayerId, const IOnlineExternalUI::FOnLoginUIClosedDelegate& InDelegate);
+	FOnlineAsyncTaskGooglePlayShowLoginUI(FOnlineSubsystemGooglePlay* InSubsystem, int InPlayerId, const FOnLoginUIClosedDelegate& InDelegate);
 
 	// FOnlineAsyncItem
 	virtual FString ToString() const override { return TEXT("ShowLoginUI"); }
@@ -36,5 +36,5 @@ private:
 	virtual void Start_OnTaskThread() override;
 
 	int PlayerId;
-	IOnlineExternalUI::FOnLoginUIClosedDelegate Delegate;
+	FOnLoginUIClosedDelegate Delegate;
 };
diff --git a/Engine/Source/Runtime/Online/Android/OnlineSubsystemGooglePlay/Private/OnlineExternalUIInterfaceGooglePlay.cpp b/Engine/Source/Runtime/Online/Android/OnlineSubsystemGooglePlay/Private/OnlineExternalUIInterfaceGooglePlay.cpp
--- a/Engine/Source/Runtime/Online/Android/OnlineSubsystemGooglePlay/Private/OnlineExternalUIInterfaceGooglePlay.cpp
+++ b/Engine/Source/Runtime/Online/Android/OnlineSubsystemGooglePlay/Private/OnlineExternalUIInterfaceGooglePlay.cpp
@@ -73,7 +73,7 @@ bool FOnlineExternalUIGooglePlay::ShowLeaderboardUI(const FString& LeaderboardNa
 	return false;
 }
 
-bool FOnlineExternalUIGooglePlay::ShowWebURL(const FString& WebURL) 
+bool FOnlineExternalUIGooglePlay::ShowWebURL(const FString& Url, const FShowWebUrlParams& ShowParams, const FOnShowWebUrlClosedDelegate& Delegate)
 {
 	return false;
 }
diff --git a/Engine/Source/Runtime/Online/Android/OnlineSubsystemGooglePlay/Private/OnlineExternalUIInterfaceGooglePlay.h b/Engine/Source/Runtime/Online/Android/OnlineSubsystemGooglePlay/Private/OnlineExternalUIInterfaceGooglePlay.h
--- a/Engine/Source/Runtime/Online/Android/OnlineSubsystemGooglePlay/Private/OnlineExternalUIInterfaceGooglePlay.h
+++ b/Engine/Source/Runtime/Online/Android/OnlineSubsystemGooglePlay/Private/OnlineExternalUIInterfaceGooglePlay.h
@@ -21,13 +21,13 @@ private:
 public:
 	FOnlineExternalUIGooglePlay(FOnlineSubsystemGooglePlay* InSubsystem);
 
-	virtual bool ShowLoginUI(const int ControllerIndex, bool bShowOnlineOnly, const FOnLoginUIClosedDelegate& Delegate) override;
+	virtual bool ShowLoginUI(const int ControllerIndex, bool bShowOnlineOnly, const FOnLoginUIClosedDelegate& Delegate = FOnLoginUIClosedDelegate() = FOnLoginUIClosedDelegate()) override;
 	virtual bool ShowFriendsUI(int32 LocalUserNum) override;
 	virtual bool ShowInviteUI(int32 LocalUserNum, FName SessionMame) override;
 	virtual bool ShowAchievementsUI(int32 LocalUserNum) override;
 	virtual bool ShowLeaderboardUI(const FString& LeaderboardName) override;
-	virtual bool ShowWebURL(const FString& WebURL) override;
-	virtual bool ShowProfileUI(const FUniqueNetId& Requestor, const FUniqueNetId& Requestee, const FOnProfileUIClosedDelegate& Delegate) override;
+	virtual bool ShowWebURL(const FString& Url, const FShowWebUrlParams& ShowParams, const FOnShowWebUrlClosedDelegate& Delegate = FOnShowWebUrlClosedDelegate()) override;
+	virtual bool ShowProfileUI(const FUniqueNetId& Requestor, const FUniqueNetId& Requestee, const FOnProfileUIClosedDelegate& Delegate = FOnProfileUIClosedDelegate()) override;
 	virtual bool ShowAccountUpgradeUI(const FUniqueNetId& UniqueId) override;
 };
 
diff --git a/Engine/Source/Runtime/Online/Android/OnlineSubsystemGooglePlay/Private/OnlineSubsystemGooglePlay.cpp b/Engine/Source/Runtime/Online/Android/OnlineSubsystemGooglePlay/Private/OnlineSubsystemGooglePlay.cpp
--- a/Engine/Source/Runtime/Online/Android/OnlineSubsystemGooglePlay/Private/OnlineSubsystemGooglePlay.cpp
+++ b/Engine/Source/Runtime/Online/Android/OnlineSubsystemGooglePlay/Private/OnlineSubsystemGooglePlay.cpp
@@ -210,7 +210,7 @@ bool FOnlineSubsystemGooglePlay::IsInAppPurchasingEnabled()
 	return bEnabledIAP;
 }
 
-void FOnlineSubsystemGooglePlay::StartShowLoginUITask(int PlayerId, const IOnlineExternalUI::FOnLoginUIClosedDelegate& Delegate)
+void FOnlineSubsystemGooglePlay::StartShowLoginUITask(int PlayerId, const FOnLoginUIClosedDelegate& Delegate)
 {
 	if (AreAnyAsyncLoginTasksRunning())
 	{
@@ -254,7 +254,7 @@ void FOnlineSubsystemGooglePlay::StartLogoutTask(int32 LocalUserNum)
 	QueueAsyncTask(CurrentLogoutTask);
 }
 
-void FOnlineSubsystemGooglePlay::StartShowLoginUITask_Internal(int PlayerId, const IOnlineExternalUI::FOnLoginUIClosedDelegate& Delegate)
+void FOnlineSubsystemGooglePlay::StartShowLoginUITask_Internal(int PlayerId, const FOnLoginUIClosedDelegate& Delegate)
 {
 	check(!AreAnyAsyncLoginTasksRunning());
 
diff --git a/Engine/Source/Runtime/Online/Android/OnlineSubsystemGooglePlay/Public/OnlineSubsystemGooglePlay.h b/Engine/Source/Runtime/Online/Android/OnlineSubsystemGooglePlay/Public/OnlineSubsystemGooglePlay.h
--- a/Engine/Source/Runtime/Online/Android/OnlineSubsystemGooglePlay/Public/OnlineSubsystemGooglePlay.h
+++ b/Engine/Source/Runtime/Online/Android/OnlineSubsystemGooglePlay/Public/OnlineSubsystemGooglePlay.h
@@ -109,7 +109,7 @@ PACKAGE_SCOPE:
 	bool AreAnyAsyncLoginTasksRunning() const { return CurrentLoginTask != nullptr || CurrentShowLoginUITask != nullptr; }
 
 	/** Start a ShowLoginUI async task. Creates the GameServices object first if necessary. */
-	void StartShowLoginUITask(int PlayerId, const IOnlineExternalUI::FOnLoginUIClosedDelegate& Delegate);
+	void StartShowLoginUITask(int PlayerId, const FOnLoginUIClosedDelegate& Delegate = FOnLoginUIClosedDelegate());
 
 	/** Start a logout task if one isn't already in progress. */
 	void StartLogoutTask(int32 LocalUserNum);
@@ -123,7 +123,7 @@ private:
 	void OnActivityResult(JNIEnv *env, jobject thiz, jobject activity, jint requestCode, jint resultCode, jobject data);
 
 	/** Start a ShowLoginUI async task. */
-	void StartShowLoginUITask_Internal(int PlayerId, const IOnlineExternalUI::FOnLoginUIClosedDelegate& Delegate);
+	void StartShowLoginUITask_Internal(int PlayerId, const FOnLoginUIClosedDelegate& Delegate = FOnLoginUIClosedDelegate());
 
 	/** Online async task runnable */
 	TUniquePtr<class FOnlineAsyncTaskManagerGooglePlay> OnlineAsyncTaskThreadRunnable;
diff --git a/Engine/Source/Runtime/Online/IOS/OnlineSubsystemIOS/Private/OnlineExternalUIInterfaceIOS.cpp b/Engine/Source/Runtime/Online/IOS/OnlineSubsystemIOS/Private/OnlineExternalUIInterfaceIOS.cpp
--- a/Engine/Source/Runtime/Online/IOS/OnlineSubsystemIOS/Private/OnlineExternalUIInterfaceIOS.cpp
+++ b/Engine/Source/Runtime/Online/IOS/OnlineSubsystemIOS/Private/OnlineExternalUIInterfaceIOS.cpp
@@ -66,12 +66,12 @@ bool FOnlineExternalUIIOS::ShowLeaderboardUI( const FString& LeaderboardName )
 	return true;
 }
 
-bool FOnlineExternalUIIOS::ShowWebURL(const FString& WebURL) 
+bool FOnlineExternalUIIOS::ShowWebURL(const FString& Url, const FShowWebUrlParams& ShowParams, const FOnShowWebUrlClosedDelegate& Delegate)
 {
 	return false;
 }
 
-bool FOnlineExternalUIIOS::ShowProfileUI( const FUniqueNetId& Requestor, const FUniqueNetId& Requestee, const FOnProfileUIClosedDelegate& Delegate )
+bool FOnlineExternalUIIOS::ShowProfileUI(const FUniqueNetId& Requestor, const FUniqueNetId& Requestee, const FOnProfileUIClosedDelegate& Delegate)
 {
 	return false;
 }
diff --git a/Engine/Source/Runtime/Online/IOS/OnlineSubsystemIOS/Private/OnlineExternalUIInterfaceIOS.h b/Engine/Source/Runtime/Online/IOS/OnlineSubsystemIOS/Private/OnlineExternalUIInterfaceIOS.h
--- a/Engine/Source/Runtime/Online/IOS/OnlineSubsystemIOS/Private/OnlineExternalUIInterfaceIOS.h
+++ b/Engine/Source/Runtime/Online/IOS/OnlineSubsystemIOS/Private/OnlineExternalUIInterfaceIOS.h
@@ -16,13 +16,13 @@ PACKAGE_SCOPE:
 public:
 
 	//~ Begin IOnlineExternalUI Interface
-	virtual bool ShowLoginUI(const int ControllerIndex, bool bShowOnlineOnly, const FOnLoginUIClosedDelegate& Delegate) override;
+	virtual bool ShowLoginUI(const int ControllerIndex, bool bShowOnlineOnly, const FOnLoginUIClosedDelegate& Delegate = FOnLoginUIClosedDelegate() = FOnLoginUIClosedDelegate()) override;
 	virtual bool ShowFriendsUI(int32 LocalUserNum) override;
 	virtual bool ShowInviteUI(int32 LocalUserNum, FName SessionMame = GameSessionName) override;
 	virtual bool ShowAchievementsUI(int32 LocalUserNum) override;
 	virtual bool ShowLeaderboardUI(const FString& LeaderboardName) override;
-	virtual bool ShowWebURL(const FString& WebURL) override;
-	virtual bool ShowProfileUI(const FUniqueNetId& Requestor, const FUniqueNetId& Requestee, const FOnProfileUIClosedDelegate& Delegate) override;
+	virtual bool ShowWebURL(const FString& Url, const FShowWebUrlParams& ShowParams, const FOnShowWebUrlClosedDelegate& Delegate = FOnShowWebUrlClosedDelegate()) override;
+	virtual bool ShowProfileUI(const FUniqueNetId& Requestor, const FUniqueNetId& Requestee, const FOnProfileUIClosedDelegate& Delegate = FOnProfileUIClosedDelegate()) override;
 	virtual bool ShowAccountUpgradeUI(const FUniqueNetId& UniqueId) override;
 	//~ End IOnlineExternalUI Interface
 	
diff --git a/Engine/Source/Runtime/Online/OnlineFramework/Party/Private/Party.cpp b/Engine/Source/Runtime/Online/OnlineFramework/Party/Private/Party.cpp
--- a/Engine/Source/Runtime/Online/OnlineFramework/Party/Private/Party.cpp
+++ b/Engine/Source/Runtime/Online/OnlineFramework/Party/Private/Party.cpp
@@ -671,7 +671,7 @@ void UParty::LeavePartyInternal(const FUniqueNetId& InUserId, const FOnlineParty
 		}
 		else
 		{
-			Result = ELeavePartyCompletionResult::UnknownClientFailure;
+			Result = ELeavePartyCompletionResult::UnknownParty;
 			ErrorMsg = FString::Printf(TEXT("Party not found in LeaveParty()"));
 		}
 	}
@@ -1146,7 +1146,7 @@ void UParty::LeavePersistentParty(const FUniqueNetId& InUserId, const UPartyDele
 	{
 		UE_LOG(LogParty, Warning, TEXT("No party during LeavePersistentParty()"));
 		// mark this as true so the ensure works in the delegate
-		InCompletionDelegate.ExecuteIfBound(InUserId, ELeavePartyCompletionResult::UnknownClientFailure);
+		InCompletionDelegate.ExecuteIfBound(InUserId, ELeavePartyCompletionResult::UnknownParty);
 	}
 }
 
diff --git a/Engine/Source/Runtime/Online/OnlineFramework/Party/Private/PartyGameState.cpp b/Engine/Source/Runtime/Online/OnlineFramework/Party/Private/PartyGameState.cpp
--- a/Engine/Source/Runtime/Online/OnlineFramework/Party/Private/PartyGameState.cpp
+++ b/Engine/Source/Runtime/Online/OnlineFramework/Party/Private/PartyGameState.cpp
@@ -337,17 +337,30 @@ void UPartyGameState::HandlePartyMemberJoined(const FUniqueNetId& InMemberId)
  	FUniqueNetIdRepl MemberId(IdRef);
  	if (MemberId.IsValid())
 	{
-		UPartyMemberState* PartyMemberState = CreateNewPartyMember(InMemberId);
-		if (ensure(PartyMemberState))
+		UPartyMemberState** NewPartyMemberPtr = PartyMembersState.Find(MemberId);
+		UPartyMemberState* NewPartyMember = NewPartyMemberPtr ? *NewPartyMemberPtr : nullptr;
+		if (!NewPartyMember)
 		{
-			PartyMembersState.Add(MemberId, PartyMemberState);
+			NewPartyMember = CreateNewPartyMember(InMemberId);
+			if (NewPartyMember)
+			{
+				PartyMembersState.Add(MemberId, NewPartyMember);
+			}
+		}
 
-			UParty* Party = GetPartyOuter();
-			check(Party);
-			Party->OnPartyMemberJoined().Broadcast(this, MemberId);
-			PartyMemberState->bHasAnnouncedJoin = true;
-			UpdateAcceptingMembers();
-		}	
+		if (ensure(NewPartyMember))
+		{
+			if (!NewPartyMember->bHasAnnouncedJoin)
+			{
+				// Both local and remote players will announce joins
+				UParty* Party = GetPartyOuter();
+				check(Party);
+				Party->OnPartyMemberJoined().Broadcast(this, MemberId);
+				NewPartyMember->bHasAnnouncedJoin = true;
+			}
+		}
+
+		UpdateAcceptingMembers();
 	}
 }
 
@@ -447,45 +460,44 @@ void UPartyGameState::HandlePartyMemberDataReceived(const FUniqueNetId& InMember
 	if (PartyInt.IsValid())
 	{
 		FUniqueNetIdRepl UniqueId(InMemberId.AsShared());
-		UPartyMemberState** CurrentPartyMemberDataPtr = PartyMembersState.Find(UniqueId);
-		UPartyMemberState* CurrentPartyMemberData = CurrentPartyMemberDataPtr ? *CurrentPartyMemberDataPtr : nullptr;
-		if (!CurrentPartyMemberData)
+		UPartyMemberState** CurrentPartyMemberPtr = PartyMembersState.Find(UniqueId);
+		UPartyMemberState* CurrentPartyMember = CurrentPartyMemberPtr ? *CurrentPartyMemberPtr : nullptr;
+		if (!CurrentPartyMember)
 		{
-			CurrentPartyMemberData = CreateNewPartyMember(InMemberId);
-			if (CurrentPartyMemberData)
+			CurrentPartyMember = CreateNewPartyMember(InMemberId);
+			if (CurrentPartyMember)
 			{
-				PartyMembersState.Add(UniqueId, CurrentPartyMemberData);
+				PartyMembersState.Add(UniqueId, CurrentPartyMember);
 			}
 		}
 
-		if (ensure(CurrentPartyMemberData))
+		if (ensure(CurrentPartyMember))
 		{
-			if (!CurrentPartyMemberData->bHasAnnouncedJoin)
+			if (!CurrentPartyMember->bHasAnnouncedJoin)
 			{
 				// Both local and remote players will announce joins
 				UParty* Party = GetPartyOuter();
 				check(Party);
 				Party->OnPartyMemberJoined().Broadcast(this, UniqueId);
-				CurrentPartyMemberData->bHasAnnouncedJoin = true;
+				CurrentPartyMember->bHasAnnouncedJoin = true;
 			}
 
-			check(CurrentPartyMemberData->MemberStateRefDef && CurrentPartyMemberData->MemberStateRefScratch);
+			check(CurrentPartyMember->MemberStateRefDef && CurrentPartyMember->MemberStateRefScratch);
 
 			// Copy out the old data
-			ensure(CurrentPartyMemberData->MemberStateRefDef->GetCppStructOps()->Copy(CurrentPartyMemberData->MemberStateRefScratch, CurrentPartyMemberData->MemberStateRef, 1));
-			if (FVariantDataConverter::VariantMapToUStruct(InPartyMemberData->KeyValAttrs, CurrentPartyMemberData->MemberStateRefDef, CurrentPartyMemberData->MemberStateRef, 0, CPF_Transient | CPF_RepSkip))
+			ensure(CurrentPartyMember->MemberStateRefDef->GetCppStructOps()->Copy(CurrentPartyMember->MemberStateRefScratch, CurrentPartyMember->MemberStateRef, 1));
+			if (FVariantDataConverter::VariantMapToUStruct(InPartyMemberData->KeyValAttrs, CurrentPartyMember->MemberStateRefDef, CurrentPartyMember->MemberStateRef, 0, CPF_Transient | CPF_RepSkip))
 			{
 				// Broadcast property changes
-				CurrentPartyMemberData->ComparePartyMemberData(*CurrentPartyMemberData->MemberStateRefScratch);
-				OnPartyMemberDataChanged().Broadcast(CurrentPartyMemberData->UniqueId, CurrentPartyMemberData);
+				CurrentPartyMember->ComparePartyMemberData(*CurrentPartyMember->MemberStateRefScratch);
+				OnPartyMemberDataChanged().Broadcast(CurrentPartyMember->UniqueId, CurrentPartyMember);
 			}
 			else
 			{
 				UE_LOG(LogParty, Warning, TEXT("[%s] Failed to serialize party member data!"), *InMemberId.ToString());
-				ensure(CurrentPartyMemberData->MemberStateRefDef->GetCppStructOps()->Copy(CurrentPartyMemberData->MemberStateRef, CurrentPartyMemberData->MemberStateRefScratch, 1));
+				ensure(CurrentPartyMember->MemberStateRefDef->GetCppStructOps()->Copy(CurrentPartyMember->MemberStateRef, CurrentPartyMember->MemberStateRefScratch, 1));
 			}
 		}
-
 	}
 }
 
diff --git a/Engine/Source/Runtime/Online/OnlineFramework/Party/Public/PartyGameState.h b/Engine/Source/Runtime/Online/OnlineFramework/Party/Public/PartyGameState.h
--- a/Engine/Source/Runtime/Online/OnlineFramework/Party/Public/PartyGameState.h
+++ b/Engine/Source/Runtime/Online/OnlineFramework/Party/Public/PartyGameState.h
@@ -485,7 +485,7 @@ protected:
 	 * @param InMemberId id of the leaving member
 	 * @param Reason reason the member left
 	 */
-	void HandlePartyMemberLeft(const FUniqueNetId& InMemberId, EMemberExitedReason Reason);
+	virtual void HandlePartyMemberLeft(const FUniqueNetId& InMemberId, EMemberExitedReason Reason);
 
 	/**
 	 * Called for all existing party members when an existing member is promoted to leader
diff --git a/Engine/Source/Runtime/Online/OnlineSubsystem/Public/Interfaces/OnlineExternalUIInterface.h b/Engine/Source/Runtime/Online/OnlineSubsystem/Public/Interfaces/OnlineExternalUIInterface.h
--- a/Engine/Source/Runtime/Online/OnlineSubsystem/Public/Interfaces/OnlineExternalUIInterface.h
+++ b/Engine/Source/Runtime/Online/OnlineSubsystem/Public/Interfaces/OnlineExternalUIInterface.h
@@ -12,6 +12,65 @@
 DECLARE_MULTICAST_DELEGATE_OneParam(FOnExternalUIChange, bool);
 typedef FOnExternalUIChange::FDelegate FOnExternalUIChangeDelegate;
 
+/**
+ * Delegate executed when the external login UI has been closed.
+ *
+ * @param UniqueId The unique id of the user who signed in. Null if no user signed in.
+ * @param ControllerIndex The controller index of the controller that activated the login UI.
+ */
+DECLARE_DELEGATE_TwoParams(FOnLoginUIClosedDelegate, TSharedPtr<const FUniqueNetId>, const int);
+
+/**
+ * Delegate executed when the web url UI has been closed
+ *
+ * @param FinalUrl the url that was used as the final redirect before closing
+ */
+DECLARE_DELEGATE_OneParam(FOnShowWebUrlClosedDelegate, const FString& /*FinalUrl*/);
+
+/** 
+ * Delegate executed when the user profile UI has been closed. 
+ */
+DECLARE_DELEGATE(FOnProfileUIClosedDelegate);
+
+/** Parameters used to show a web UI */
+struct FShowWebUrlParams
+{
+	/** presented without a frame if embedded enabled */
+	bool bEmbedded;
+	/** x offset in pixels from top left */
+	int32 OffsetX;
+	/** y offset in pixels from top left */
+	int32 OffsetY;
+	/** x size in pixels */
+	int32 SizeX;
+	/** y size in pixels */
+	int32 SizeY;
+	/** if specified then restricted to only navigate within these domains */
+	TArray<FString> AllowedDomains;
+
+	/**
+	 * Constructor
+	 */
+	FShowWebUrlParams(bool InbEmbedded, int32 InOffsetX, int32 InOffsetY, int32 InSizeX, int32 InSizeY)
+		: bEmbedded(InbEmbedded)
+		, OffsetX(InOffsetX)
+		, OffsetY(InOffsetY)
+		, SizeX(InSizeX)
+		, SizeY(InSizeY)
+	{}
+
+	/**
+	 * Default Constructor
+	 */
+	FShowWebUrlParams()
+		: bEmbedded(false)
+		, OffsetX(0)
+		, OffsetY(0)
+		, SizeX(0)
+		, SizeY(0)
+	{}
+};
+
 /** 
  * Interface definition for the online services external UIs
  * Any online service that provides extra UI overlays will implement the relevant functions
@@ -24,14 +83,6 @@ protected:
 public:
 	virtual ~IOnlineExternalUI() {};
 
-	/**
-	 * Delegate executed when the external login UI has been closed.
-	 *
-	 * @param UniqueId The unique id of the user who signed in. Null if no user signed in.
-	 * @param ControllerIndex The controller index of the controller that activated the login UI.
-	 */
-	DECLARE_DELEGATE_TwoParams(FOnLoginUIClosedDelegate, TSharedPtr<const FUniqueNetId>, const int);
-
 	/**
 	 * Displays the UI that prompts the user for their login credentials. Each
 	 * platform handles the authentication of the user's data.
@@ -43,7 +94,7 @@ public:
  	 *
 	 * @return true if it was able to show the UI, false if it failed
 	 */
-	virtual bool ShowLoginUI(const int ControllerIndex, bool bShowOnlineOnly, const FOnLoginUIClosedDelegate& Delegate) = 0;
+	virtual bool ShowLoginUI(const int ControllerIndex, bool bShowOnlineOnly, const FOnLoginUIClosedDelegate& Delegate = FOnLoginUIClosedDelegate() = FOnLoginUIClosedDelegate()) = 0;
 
 	/**
 	 * Displays the UI that shows a user's list of friends
@@ -88,10 +139,7 @@ public:
 	 *
 	 * @return true if it was able to show the UI, false if it failed
 	 */
-	virtual bool ShowWebURL(const FString& WebURL) = 0;
-
-	/** Delegate executed when the user profile UI has been closed. */
-	DECLARE_DELEGATE(FOnProfileUIClosedDelegate);
+	virtual bool ShowWebURL(const FString& Url, const FShowWebUrlParams& ShowParams, const FOnShowWebUrlClosedDelegate& Delegate = FOnShowWebUrlClosedDelegate()) = 0;
 
 	/**
 	 * Displays a user's profile card.
@@ -101,7 +149,7 @@ public:
 	 *
 	 * @return true if it was able to show the UI, false if it failed
 	 */
-	virtual bool ShowProfileUI(const FUniqueNetId& Requestor, const FUniqueNetId& Requestee, const FOnProfileUIClosedDelegate& Delegate) = 0;
+	virtual bool ShowProfileUI(const FUniqueNetId& Requestor, const FUniqueNetId& Requestee, const FOnProfileUIClosedDelegate& Delegate = FOnProfileUIClosedDelegate()) = 0;
 
 	/**
 	* Displays a system dialog to purchase user account upgrades.  e.g. PlaystationPlus, XboxLive GOLD, etc.
diff --git a/Engine/Source/Runtime/Online/OnlineSubsystemSteam/Private/OnlineExternalUIInterfaceSteam.cpp b/Engine/Source/Runtime/Online/OnlineSubsystemSteam/Private/OnlineExternalUIInterfaceSteam.cpp
--- a/Engine/Source/Runtime/Online/OnlineSubsystemSteam/Private/OnlineExternalUIInterfaceSteam.cpp
+++ b/Engine/Source/Runtime/Online/OnlineSubsystemSteam/Private/OnlineExternalUIInterfaceSteam.cpp
@@ -29,6 +29,10 @@ void FOnlineAsyncEventSteamExternalUITriggered::TriggerDelegates()
 		FOnlineExternalUISteamPtr ExternalUISteam = StaticCastSharedPtr<FOnlineExternalUISteam>(ExternalUIInterface);
 		ExternalUISteam->ProfileUIClosedDelegate.ExecuteIfBound();
 		ExternalUISteam->ProfileUIClosedDelegate.Unbind();
+
+		//@todo samz - obtain final url
+		ExternalUISteam->ShowWebUrlClosedDelegate.ExecuteIfBound(TEXT(""));
+		ExternalUISteam->ShowWebUrlClosedDelegate.Unbind();
 	}
 }
 
@@ -66,21 +70,22 @@ bool FOnlineExternalUISteam::ShowLeaderboardUI(const FString& LeaderboardName)
 	return false;
 }
 
-bool FOnlineExternalUISteam::ShowWebURL(const FString& WebURL)
+bool FOnlineExternalUISteam::ShowWebURL(const FString& Url, const FShowWebUrlParams& ShowParams, const FOnShowWebUrlClosedDelegate& Delegate)
 {
-	if (!WebURL.StartsWith(TEXT("https://")))
+	if (!Url.StartsWith(TEXT("https://")))
 	{
-		SteamFriends()->ActivateGameOverlayToWebPage(TCHAR_TO_UTF8(*FString::Printf(TEXT("https://%s"), *WebURL)));
+		SteamFriends()->ActivateGameOverlayToWebPage(TCHAR_TO_UTF8(*FString::Printf(TEXT("https://%s"), *Url)));
 	}
 	else
 	{
-		SteamFriends()->ActivateGameOverlayToWebPage(TCHAR_TO_UTF8(*WebURL));
+		SteamFriends()->ActivateGameOverlayToWebPage(TCHAR_TO_UTF8(*Url));
 	}
-	
+
+	ShowWebUrlClosedDelegate = Delegate;
 	return true;
 }
 
-bool FOnlineExternalUISteam::ShowProfileUI( const FUniqueNetId& Requestor, const FUniqueNetId& Requestee, const FOnProfileUIClosedDelegate& Delegate )
+bool FOnlineExternalUISteam::ShowProfileUI(const FUniqueNetId& Requestor, const FUniqueNetId& Requestee, const FOnProfileUIClosedDelegate& Delegate)
 {
 	SteamFriends()->ActivateGameOverlayToUser(TCHAR_TO_UTF8(TEXT("steamid")), (const FUniqueNetIdSteam&)Requestee);
 
diff --git a/Engine/Source/Runtime/Online/OnlineSubsystemSteam/Private/OnlineExternalUIInterfaceSteam.h b/Engine/Source/Runtime/Online/OnlineSubsystemSteam/Private/OnlineExternalUIInterfaceSteam.h
--- a/Engine/Source/Runtime/Online/OnlineSubsystemSteam/Private/OnlineExternalUIInterfaceSteam.h
+++ b/Engine/Source/Runtime/Online/OnlineSubsystemSteam/Private/OnlineExternalUIInterfaceSteam.h
@@ -70,6 +70,7 @@ PACKAGE_SCOPE:
 
 	/** Triggered when the steam overlay is closed if it was opened via ShowProfileUI. Delegate will be unbound after it is executed. */
 	FOnProfileUIClosedDelegate ProfileUIClosedDelegate;
+	FOnShowWebUrlClosedDelegate ShowWebUrlClosedDelegate;
 
 public:
 
@@ -78,13 +79,13 @@ public:
 	}
 
 	// IOnlineExternalUI
-	virtual bool ShowLoginUI(const int ControllerIndex, bool bShowOnlineOnly, const FOnLoginUIClosedDelegate& Delegate) override;
+	virtual bool ShowLoginUI(const int ControllerIndex, bool bShowOnlineOnly, const FOnLoginUIClosedDelegate& Delegate = FOnLoginUIClosedDelegate()) override;
 	virtual bool ShowFriendsUI(int32 LocalUserNum) override;
 	virtual bool ShowInviteUI(int32 LocalUserNum, FName SessionMame = GameSessionName) override;
 	virtual bool ShowAchievementsUI(int32 LocalUserNum) override;
 	virtual bool ShowLeaderboardUI(const FString& LeaderboardName) override;
-	virtual bool ShowWebURL(const FString& WebURL) override;
-	virtual bool ShowProfileUI(const FUniqueNetId& Requestor, const FUniqueNetId& Requestee, const FOnProfileUIClosedDelegate& Delegate) override;
+	virtual bool ShowWebURL(const FString& Url, const FShowWebUrlParams& ShowParams, const FOnShowWebUrlClosedDelegate& Delegate = FOnShowWebUrlClosedDelegate()) override;
+	virtual bool ShowProfileUI(const FUniqueNetId& Requestor, const FUniqueNetId& Requestee, const FOnProfileUIClosedDelegate& Delegate = FOnProfileUIClosedDelegate()) override;
 	virtual bool ShowAccountUpgradeUI(const FUniqueNetId& UniqueId) override;
 };
 
diff --git a/Engine/Source/Runtime/Online/OnlineSubsystemUtils/Private/ShowLoginUICallbackProxy.cpp b/Engine/Source/Runtime/Online/OnlineSubsystemUtils/Private/ShowLoginUICallbackProxy.cpp
--- a/Engine/Source/Runtime/Online/OnlineSubsystemUtils/Private/ShowLoginUICallbackProxy.cpp
+++ b/Engine/Source/Runtime/Online/OnlineSubsystemUtils/Private/ShowLoginUICallbackProxy.cpp
@@ -53,7 +53,7 @@ void UShowLoginUICallbackProxy::Activate()
 	}
 		
 	const bool bWaitForDelegate = OnlineExternalUI->ShowLoginUI(LocalPlayer->GetControllerId(), false,
-		IOnlineExternalUI::FOnLoginUIClosedDelegate::CreateUObject(this, &UShowLoginUICallbackProxy::OnShowLoginUICompleted));
+		FOnLoginUIClosedDelegate::CreateUObject(this, &UShowLoginUICallbackProxy::OnShowLoginUICompleted));
 
 	if (!bWaitForDelegate)
 	{
diff --git a/Engine/Source/Runtime/Online/OnlineSubsystemUtils/Private/Tests/TestExternalUIInterface.cpp b/Engine/Source/Runtime/Online/OnlineSubsystemUtils/Private/Tests/TestExternalUIInterface.cpp
--- a/Engine/Source/Runtime/Online/OnlineSubsystemUtils/Private/Tests/TestExternalUIInterface.cpp
+++ b/Engine/Source/Runtime/Online/OnlineSubsystemUtils/Private/Tests/TestExternalUIInterface.cpp
@@ -89,7 +89,7 @@ bool FTestExternalUIInterface::TestLoginUI()
 		return false;
 	}
 
-	bool bShowingUI = ExternalUI->ShowLoginUI(0, true, IOnlineExternalUI::FOnLoginUIClosedDelegate::CreateRaw(this, &FTestExternalUIInterface::OnLoginUIClosed));
+	bool bShowingUI = ExternalUI->ShowLoginUI(0, true, FOnLoginUIClosedDelegate::CreateRaw(this, &FTestExternalUIInterface::OnLoginUIClosed));
 	UE_LOG(LogOnline, Log, TEXT("TestLoginUI bShowingUI: %d"), bShowingUI);
 	return bShowingUI;
 }
@@ -143,7 +143,11 @@ bool FTestExternalUIInterface::TestWebURL()
 		return false;
 	}
 
-	bool bShowingUI = ExternalUI->ShowWebURL("https://www.unrealengine.com");
+	bool bShowingUI = ExternalUI->ShowWebURL(
+		TEXT("https://www.unrealengine.com"),
+		FShowWebUrlParams(),
+		FOnShowWebUrlClosedDelegate::CreateRaw(this, &FTestExternalUIInterface::OnShowWebUrlClosed));
+
 	UE_LOG(LogOnline, Log, TEXT("TestWebURL bShowingUI: %d"), bShowingUI);
 	return bShowingUI;
 }
@@ -161,7 +165,7 @@ bool FTestExternalUIInterface::TestProfileUI()
 	bool bShowingUI = ExternalUI->ShowProfileUI(
 		*UserId.Get(),
 		*UserId.Get(),
-		IOnlineExternalUI::FOnProfileUIClosedDelegate::CreateRaw(this, &FTestExternalUIInterface::OnProfileUIClosed));
+		FOnProfileUIClosedDelegate::CreateRaw(this, &FTestExternalUIInterface::OnProfileUIClosed));
 
 	UE_LOG(LogOnline, Log, TEXT("TestProfileUI bShowingUI: %d"), bShowingUI);
 	return bShowingUI;
@@ -187,4 +191,9 @@ void FTestExternalUIInterface::OnLoginUIClosed(TSharedPtr<const FUniqueNetId> Lo
 void FTestExternalUIInterface::OnProfileUIClosed()
 {
 	UE_LOG(LogOnline, Log, TEXT("Profile UI closed by user."));
+}
+
+void FTestExternalUIInterface::OnShowWebUrlClosed(const FString& FinalUrl)
+{
+	UE_LOG(LogOnline, Log, TEXT("Show Web Url closed with FinalUrl=%s."), *FinalUrl);
 }
\ No newline at end of file
diff --git a/Engine/Source/Runtime/Online/OnlineSubsystemUtils/Private/Tests/TestExternalUIInterface.h b/Engine/Source/Runtime/Online/OnlineSubsystemUtils/Private/Tests/TestExternalUIInterface.h
--- a/Engine/Source/Runtime/Online/OnlineSubsystemUtils/Private/Tests/TestExternalUIInterface.h
+++ b/Engine/Source/Runtime/Online/OnlineSubsystemUtils/Private/Tests/TestExternalUIInterface.h
@@ -76,6 +76,9 @@ namespace ETestExternalUIInterfaceState
 	/** Delegate executed when the user profile UI has been closed. */
 	void OnProfileUIClosed();
 
+	/** Delegate executed when the show web UI has been closed. */
+	void OnShowWebUrlClosed(const FString& FinalUrl);
+
  public:
 
 	/**
diff --git a/Engine/Source/Runtime/Renderer/Private/Renderer.cpp b/Engine/Source/Runtime/Renderer/Private/Renderer.cpp
--- a/Engine/Source/Runtime/Renderer/Private/Renderer.cpp
+++ b/Engine/Source/Runtime/Renderer/Private/Renderer.cpp
@@ -32,6 +32,11 @@ void FRendererModule::SceneRenderTargetsSetBufferSize(uint32 SizeX, uint32 SizeY
 	FSceneRenderTargets::Get_Todo_PassContext().UpdateRHI();
 }
 
+void FRendererModule::InitializeSystemTextures(FRHICommandListImmediate& RHICmdList)
+{
+	GSystemTextures.InitializeTextures(RHICmdList, GMaxRHIFeatureLevel);
+}
+
 void FRendererModule::DrawTileMesh(FRHICommandListImmediate& RHICmdList, const FSceneView& SceneView, const FMeshBatch& Mesh, bool bIsHitTesting, const FHitProxyId& HitProxyId)
 {
 	// Create an FViewInfo so we can initialize its RHI resources
diff --git a/Engine/Source/Runtime/Renderer/Private/RendererPrivate.h b/Engine/Source/Runtime/Renderer/Private/RendererPrivate.h
--- a/Engine/Source/Runtime/Renderer/Private/RendererPrivate.h
+++ b/Engine/Source/Runtime/Renderer/Private/RendererPrivate.h
@@ -31,6 +31,7 @@ public:
 	virtual uint32 GetNumDynamicLightsAffectingPrimitive(const FPrimitiveSceneInfo* PrimitiveSceneInfo,const FLightCacheInterface* LCI) override;
 	virtual void ReallocateSceneRenderTargets() override;
 	virtual void SceneRenderTargetsSetBufferSize(uint32 SizeX, uint32 SizeY) override;
+	virtual void InitializeSystemTextures(FRHICommandListImmediate& RHICmdList);
 	virtual void DrawTileMesh(FRHICommandListImmediate& RHICmdList, const FSceneView& View, const FMeshBatch& Mesh, bool bIsHitTesting, const FHitProxyId& HitProxyId) override;
 	virtual void RenderTargetPoolFindFreeElement(FRHICommandListImmediate& RHICmdList, const FPooledRenderTargetDesc& Desc, TRefCountPtr<IPooledRenderTarget> &Out, const TCHAR* InDebugName) override;
 	virtual void TickRenderTargetPool() override;
diff --git a/Engine/Source/Runtime/Renderer/Public/RendererInterface.h b/Engine/Source/Runtime/Renderer/Public/RendererInterface.h
--- a/Engine/Source/Runtime/Renderer/Public/RendererInterface.h
+++ b/Engine/Source/Runtime/Renderer/Public/RendererInterface.h
@@ -548,6 +548,8 @@ public:
 	/** Sets the buffer size of the render targets. */
 	virtual void SceneRenderTargetsSetBufferSize(uint32 SizeX, uint32 SizeY) = 0;
 
+	virtual void InitializeSystemTextures(FRHICommandListImmediate& RHICmdList) = 0;
+
 	/** Draws a tile mesh element with the specified view. */
 	virtual void DrawTileMesh(FRHICommandListImmediate& RHICmdList, const FSceneView& View, const FMeshBatch& Mesh, bool bIsHitTesting, const class FHitProxyId& HitProxyId) = 0;
 
diff --git a/Engine/Source/Runtime/Slate/Private/Framework/Application/MenuStack.cpp b/Engine/Source/Runtime/Slate/Private/Framework/Application/MenuStack.cpp
--- a/Engine/Source/Runtime/Slate/Private/Framework/Application/MenuStack.cpp
+++ b/Engine/Source/Runtime/Slate/Private/Framework/Application/MenuStack.cpp
@@ -74,21 +74,33 @@ namespace
 		SLATE_BEGIN_ARGS(SMenuContentWrapper) 
 				: _MenuContent()
 				, _OnKeyDown()
+				, _OptionalMinMenuWidth()
+				, _OptionalMinMenuHeight()
 			{}
 
 			SLATE_DEFAULT_SLOT(FArguments, MenuContent)
 			SLATE_EVENT(FOnKeyDown, OnKeyDown)
 			SLATE_EVENT(FOnMenuLostFocus, OnMenuLostFocus)
+			SLATE_ARGUMENT(FOptionalSize, OptionalMinMenuWidth)
+			SLATE_ARGUMENT(FOptionalSize, OptionalMinMenuHeight)
 		SLATE_END_ARGS()
 
 		/** Construct this widget */
 		void Construct(const FArguments& InArgs)
 		{
+			// The visibility of the content wrapper should match that of the provided content
+			Visibility = AccessWidgetVisibilityAttribute(InArgs._MenuContent.Widget);
+
 			OnKeyDownDelegate = InArgs._OnKeyDown;
 			OnMenuLostFocus = InArgs._OnMenuLostFocus;
 			ChildSlot
 			[
-				InArgs._MenuContent.Widget
+				SNew(SBox)
+				.MinDesiredWidth(InArgs._OptionalMinMenuWidth)
+				.MaxDesiredHeight(InArgs._OptionalMinMenuHeight)
+				[
+					InArgs._MenuContent.Widget
+				]
 			];
 		}
 
@@ -658,14 +670,11 @@ TSharedRef<SWidget> FMenuStack::WrapContent(TSharedRef<SWidget> InContent, FOpti
 	return SNew(SMenuContentWrapper)
 		.OnKeyDown_Static(&OnMenuKeyDown)
 		.OnMenuLostFocus_Raw(this, &FMenuStack::OnMenuContentLostFocus)
+		.OptionalMinMenuWidth(OptionalMinWidth)
+		.OptionalMinMenuHeight(OptionalMinHeight)
 		.MenuContent()
 		[
-			SNew(SBox)
-			.MinDesiredWidth(OptionalMinWidth)
-			.MaxDesiredHeight(OptionalMinHeight)
-			[
-				InContent
-			]
+			InContent
 		];
 }
 
diff --git a/Engine/Source/Runtime/Slate/Private/Widgets/Input/SMenuAnchor.cpp b/Engine/Source/Runtime/Slate/Private/Widgets/Input/SMenuAnchor.cpp
--- a/Engine/Source/Runtime/Slate/Private/Widgets/Input/SMenuAnchor.cpp
+++ b/Engine/Source/Runtime/Slate/Private/Widgets/Input/SMenuAnchor.cpp
@@ -14,6 +14,9 @@ static FVector2D GetMenuOffsetForPlacement(const FGeometry& AllottedGeometry, EM
 		case MenuPlacement_CenteredBelowAnchor:
 			return FVector2D(-((PopupSizeLocalSpace.X / 2) - (AllottedGeometry.GetLocalSize().X / 2)), AllottedGeometry.GetLocalSize().Y);
 			break;
+		case MenuPlacement_BelowRightAnchor:
+			return FVector2D( -( PopupSizeLocalSpace.X ) - ( AllottedGeometry.GetLocalSize().X ), AllottedGeometry.GetLocalSize().Y );
+			break;
 		case MenuPlacement_ComboBox:
 			return FVector2D(0.0f, AllottedGeometry.GetLocalSize().Y);
 			break;
@@ -29,6 +32,9 @@ static FVector2D GetMenuOffsetForPlacement(const FGeometry& AllottedGeometry, EM
 		case MenuPlacement_CenteredAboveAnchor:
 			return FVector2D(-((PopupSizeLocalSpace.X / 2) - (AllottedGeometry.GetLocalSize().X / 2)), -PopupSizeLocalSpace.Y);
 			break;
+		case MenuPlacement_AboveRightAnchor:
+			return FVector2D( -( PopupSizeLocalSpace.X ) - ( AllottedGeometry.GetLocalSize().X ), -PopupSizeLocalSpace.Y );
+			break;
 		case MenuPlacement_MenuLeft:
 			return FVector2D(-PopupSizeLocalSpace.X, 0.0f);
 			break;
diff --git a/Engine/Source/Runtime/SlateCore/Public/Types/SlateEnums.h b/Engine/Source/Runtime/SlateCore/Public/Types/SlateEnums.h
--- a/Engine/Source/Runtime/SlateCore/Public/Types/SlateEnums.h
+++ b/Engine/Source/Runtime/SlateCore/Public/Types/SlateEnums.h
@@ -157,6 +157,9 @@ enum EMenuPlacement
 	/** Place the menu immediately centered below the anchor */
 	MenuPlacement_CenteredBelowAnchor UMETA(DisplayName = "Centered Below"),
 
+	/** Place the menu immediately below the anchor aligned to the right of the content */
+	MenuPlacement_BelowRightAnchor UMETA(DisplayName = "Below Right"),
+
 	/** Place the menu immediately below the anchor and match is width to the anchor's content */
 	MenuPlacement_ComboBox UMETA(DisplayName="Combo Box"),
 
@@ -172,6 +175,9 @@ enum EMenuPlacement
 	/** Place the menu immediately centered above the anchor, no transition effect */
 	MenuPlacement_CenteredAboveAnchor UMETA(DisplayName="Centered Above"),
 
+	/** Place the menu immediately above the anchor aligned to the right of the content */
+	MenuPlacement_AboveRightAnchor UMETA(DisplayName = "Above Right"),
+
 	/** Place the menu to the left of the anchor */
 	MenuPlacement_MenuLeft UMETA(DisplayName="Left"),
 
diff --git a/Engine/Source/Runtime/SlateCore/Public/Widgets/SWidget.h b/Engine/Source/Runtime/SlateCore/Public/Widgets/SWidget.h
--- a/Engine/Source/Runtime/SlateCore/Public/Widgets/SWidget.h
+++ b/Engine/Source/Runtime/SlateCore/Public/Widgets/SWidget.h
@@ -994,6 +994,16 @@ protected:
 		return Visibility.IsBound() || EnabledState.IsBound() || RenderTransform.IsBound();
 	}
 
+	/**
+	 * Protected static helper to allow widgets to access the visibility attribute of other widgets directly
+	 * 
+	 * @param Widget The widget to get the visibility attribute of
+	 */
+	static const TAttribute<EVisibility>& AccessWidgetVisibilityAttribute(const TSharedRef<SWidget>& Widget)
+	{
+		return Widget->Visibility;
+	}
+
 private:
 
 	/**
diff --git a/Engine/Source/Runtime/UMG/Private/Components/Widget.cpp b/Engine/Source/Runtime/UMG/Private/Components/Widget.cpp
--- a/Engine/Source/Runtime/UMG/Private/Components/Widget.cpp
+++ b/Engine/Source/Runtime/UMG/Private/Components/Widget.cpp
@@ -601,8 +601,18 @@ void UWidget::SetDesignerFlags(EWidgetDesignFlags::Type NewFlags)
 	DesignerFlags = ( EWidgetDesignFlags::Type )( DesignerFlags | NewFlags );
 }
 
+void UWidget::SetDisplayLabel(const FString& InDisplayLabel)
+{
+	DisplayLabel = InDisplayLabel;
+}
+
 bool UWidget::IsGeneratedName() const
 {
+	if (!DisplayLabel.IsEmpty())
+	{
+		return false;
+	}
+
 	FString Name = GetName();
 
 	if (Name == GetClass()->GetName() || Name.StartsWith(GetClass()->GetName() + TEXT("_")))
@@ -650,7 +660,15 @@ FText UWidget::GetLabelText() const
 
 FText UWidget::GetDisplayNameBase() const
 {
-	return IsGeneratedName() ? GetClass()->GetDisplayNameText() : FText::FromString(GetName());
+	const bool bHasDisplayLabel = !DisplayLabel.IsEmpty();
+	if (IsGeneratedName())
+	{
+		return GetClass()->GetDisplayNameText();
+	}
+	else
+	{
+		return FText::FromString(bHasDisplayLabel ? DisplayLabel : GetName());
+	}
 }
 
 const FText UWidget::GetPaletteCategory()
diff --git a/Engine/Source/Runtime/UMG/Private/Slate/SMeshWidget.cpp b/Engine/Source/Runtime/UMG/Private/Slate/SMeshWidget.cpp
--- a/Engine/Source/Runtime/UMG/Private/Slate/SMeshWidget.cpp
+++ b/Engine/Source/Runtime/UMG/Private/Slate/SMeshWidget.cpp
@@ -104,13 +104,22 @@ void SMeshWidget::ClearRuns(int32 NumRuns)
 }
 
 static const FName SlateRHIModuleName("SlateRHIRenderer");
-TSharedPtr<FSlateInstanceBufferUpdate> SMeshWidget::BeginPerInstanceBufferUpdate(uint32 MeshId, int32 InitialSize)
+void SMeshWidget::EnableInstancing(uint32 MeshId, int32 InitialSize)
 {
 	if (!RenderData[MeshId].PerInstanceBuffer.IsValid())
 	{
 		RenderData[MeshId].PerInstanceBuffer = FModuleManager::Get().GetModuleChecked<ISlateRHIRendererModule>(SlateRHIModuleName).CreateInstanceBuffer(InitialSize);
 	}
+}
 
+TSharedPtr<FSlateInstanceBufferUpdate> SMeshWidget::BeginPerInstanceBufferUpdate(uint32 MeshId, int32 InitialSize)
+{
+	EnableInstancing(MeshId, InitialSize);
+	return BeginPerInstanceBufferUpdateConst( MeshId );
+}
+
+TSharedPtr<FSlateInstanceBufferUpdate> SMeshWidget::BeginPerInstanceBufferUpdateConst(uint32 MeshId) const
+{
 	return RenderData[MeshId].PerInstanceBuffer->BeginUpdate();
 }
 
@@ -190,3 +199,16 @@ FVector2D SMeshWidget::ComputeDesiredSize(float) const
 }
 
 
+
+
+void SMeshWidget::AddReferencedObjects(FReferenceCollector& Collector)
+{
+	for (const FRenderData& SomeRenderData : RenderData)
+	{
+		if (SomeRenderData.Brush.IsValid())
+		{
+			UObject* ResourceObject = SomeRenderData.Brush->GetResourceObject();
+			Collector.AddReferencedObject(ResourceObject);
+		}
+	}
+}
\ No newline at end of file
diff --git a/Engine/Source/Runtime/UMG/Private/Slate/SlateVectorArtData.cpp b/Engine/Source/Runtime/UMG/Private/Slate/SlateVectorArtData.cpp
--- a/Engine/Source/Runtime/UMG/Private/Slate/SlateVectorArtData.cpp
+++ b/Engine/Source/Runtime/UMG/Private/Slate/SlateVectorArtData.cpp
@@ -5,8 +5,11 @@
 #include "StaticMeshResources.h"
 #include "SlateVectorArtData.h"
 
-static void StaticMeshToSlateRenderData(const UStaticMesh& DataSource, TArray<FSlateMeshVertex>& OutSlateVerts, TArray<uint32>& OutIndexes)
+static void StaticMeshToSlateRenderData(const UStaticMesh& DataSource, TArray<FSlateMeshVertex>& OutSlateVerts, TArray<uint32>& OutIndexes, FVector2D& OutExtentMin, FVector2D& OutExtentMax )
 {
+	OutExtentMin = FVector2D(FLT_MAX, FLT_MAX);
+	OutExtentMax = FVector2D(-FLT_MAX, -FLT_MAX);
+
 	const FStaticMeshLODResources& LOD = DataSource.RenderData->LODResources[0];
 	const int32 NumSections = LOD.Sections.Num();
 	if (NumSections > 1)
@@ -32,6 +35,10 @@ static void StaticMeshToSlateRenderData(const UStaticMesh& DataSource, TArray<FS
 			{
 				// Copy Position
 				const FVector& Position = LOD.PositionVertexBuffer.VertexPosition(i);
+				OutExtentMin.X = FMath::Min(Position.X, OutExtentMin.X);
+				OutExtentMin.Y = FMath::Min(Position.Y, OutExtentMin.Y);
+				OutExtentMax.X = FMath::Max(Position.X, OutExtentMax.X);
+				OutExtentMax.Y = FMath::Max(Position.Y, OutExtentMax.Y);
 				
 				// Copy Color
 				FColor Color = (LOD.ColorVertexBuffer.GetNumVertices() > 0) ? LOD.ColorVertexBuffer.VertexColor(i) : FColor::White;
@@ -120,18 +127,24 @@ UMaterialInterface* USlateVectorArtData::GetMaterial() const
 
 UMaterialInstanceDynamic* USlateVectorArtData::ConvertToMaterialInstanceDynamic()
 {
-	UMaterialInstanceDynamic* NewMID = UMaterialInstanceDynamic::Create(Material, this);
-	Material = NewMID;
-	return NewMID;
+	EnsureValidData();
+	UMaterialInstanceDynamic* ExistingMID = Cast<UMaterialInstanceDynamic>(Material);
+	if (ExistingMID == nullptr)
+	{
+		UMaterialInstanceDynamic* NewMID = UMaterialInstanceDynamic::Create(Material, this);
+		Material = NewMID;
+		return NewMID;
+	}
+	else
+	{
+		return ExistingMID;
+	}
 }
 
 void USlateVectorArtData::EnsureValidData()
 {
 #if WITH_EDITORONLY_DATA
-	if (VertexData.Num() == 0 || IndexData.Num() == 0 || Material == nullptr)
-	{
-		InitFromStaticMesh(*MeshAsset);
-	}
+	InitFromStaticMesh(*MeshAsset);
 #endif
 }
 
@@ -144,10 +157,29 @@ void USlateVectorArtData::PreSave()
 #if WITH_EDITORONLY_DATA
 void USlateVectorArtData::InitFromStaticMesh(const UStaticMesh& InSourceMesh)
 {
-	Material = InSourceMesh.GetMaterial(0);
+	if (SourceMaterial != InSourceMesh.GetMaterial(0))
+	{
+		SourceMaterial = InSourceMesh.GetMaterial(0);
+		Material = SourceMaterial;
+	}
 
 	ensureMsgf(Material != nullptr, TEXT("USlateVectorArtData::InitFromStaticMesh() expected %s to have a material assigned."), *InSourceMesh.GetFullName());
 
-	StaticMeshToSlateRenderData(InSourceMesh, VertexData, IndexData);
+	StaticMeshToSlateRenderData(InSourceMesh, VertexData, IndexData, ExtentMin, ExtentMax);
+}
+#endif
+
+FVector2D USlateVectorArtData::GetDesiredSize() const
+{
+	return GetExtentMax() - GetExtentMin();
+}
+
+FVector2D USlateVectorArtData::GetExtentMin() const
+{
+	return ExtentMin;
 }
-#endif
\ No newline at end of file
+
+FVector2D USlateVectorArtData::GetExtentMax() const
+{
+	return ExtentMax;
+}
\ No newline at end of file
diff --git a/Engine/Source/Runtime/UMG/Private/Slate/SlateVectorArtInstanceData.cpp b/Engine/Source/Runtime/UMG/Private/Slate/SlateVectorArtInstanceData.cpp
new file mode 100644
--- /dev/null
+++ b/Engine/Source/Runtime/UMG/Private/Slate/SlateVectorArtInstanceData.cpp
@@ -0,0 +1,32 @@
+// Copyright 1998-2015 Epic Games, Inc. All Rights Reserved.
+
+#include "UMGPrivatePCH.h"
+#include "SlateVectorArtInstanceData.h"
+
+void FSlateVectorArtInstanceData::SetPositionFixedPoint16(FVector2D Position)
+{
+	// We are effectively doing fixed point positioning
+	// With 13 bits for the whole part and 3 bits of decimal precision.
+	// xxxxxxxxxxxxx.xxx
+	//     [0..8191].[0..7]
+	{
+		uint32 CurrentDataX = static_cast<uint32>(GetData().X) & 0xFFFF0000;
+		GetData().X = CurrentDataX | FMath::RoundToInt(Position.X*8.0f) << 0;
+	}
+
+	{
+		uint32 CurrentDataY = static_cast<uint32>(GetData().Y) & 0xFFFF0000;
+		GetData().Y = CurrentDataY | FMath::RoundToInt(Position.Y*8.0f) << 0;
+	}
+}
+
+void FSlateVectorArtInstanceData::SetPosition(FVector2D Position)
+{
+	GetData().X = Position.X;
+	GetData().Y = Position.Y;
+}
+
+void FSlateVectorArtInstanceData::SetScale(float Scale)
+{
+	GetData().Z = Scale;
+}
