Copying //UE4/Dev-Sequencer to Dev-Main (//UE4/Dev-Main)

diff --git a/Engine/Source/Runtime/Core/Private/Tests/Misc/QueueTest.cpp b/Engine/Source/Runtime/Core/Private/Tests/Misc/QueueTest.cpp
--- a/Engine/Source/Runtime/Core/Private/Tests/Misc/QueueTest.cpp
+++ b/Engine/Source/Runtime/Core/Private/Tests/Misc/QueueTest.cpp
@@ -38,5 +38,16 @@ bool FQueueTest::RunTest( const FString& Parameters )
 
 	TestTrue(TEXT("After removing all items, the queue must be empty"), q2_1.IsEmpty());
 
+	// emptying
+	TQueue<int32> q3_1;
+	{
+		q3_1.Enqueue(1);
+		q3_1.Enqueue(2);
+		q3_1.Enqueue(3);
+		q3_1.Empty();
+	}
+
+	TestTrue(TEXT("An emptied queue must be empty"), q3_1.IsEmpty());
+
 	return true;
 }
diff --git a/Engine/Source/Runtime/Core/Private/Tests/Misc/TypeContainerTest.cpp b/Engine/Source/Runtime/Core/Private/Tests/Misc/TypeContainerTest.cpp
--- a/Engine/Source/Runtime/Core/Private/Tests/Misc/TypeContainerTest.cpp
+++ b/Engine/Source/Runtime/Core/Private/Tests/Misc/TypeContainerTest.cpp
@@ -11,7 +11,7 @@
 #include "AutomationTest.h"
 
 
-IMPLEMENT_SIMPLE_AUTOMATION_TEST(FTypeContainerTest, "System.Core.Misc.TypeContainer", EAutomationTestFlags::Disabled | EAutomationTestFlags::ApplicationContextMask | EAutomationTestFlags::SmokeFilter)
+IMPLEMENT_SIMPLE_AUTOMATION_TEST(FTypeContainerTest, "System.Core.Misc.TypeContainer", EAutomationTestFlags::ApplicationContextMask | EAutomationTestFlags::EngineFilter)
 
 
 /* Helpers
@@ -40,37 +40,41 @@ struct FStrawberry : public IBerry
 	virtual FString Name() override { return TEXT("Strawberry"); }
 };
 
+template<ESPMode Mode = ESPMode::Fast>
 struct ISmoothie
 {
 	virtual ~ISmoothie() { }
-	virtual TSharedRef<IBerry> GetBerry() = 0;
-	virtual TSharedRef<IFruit> GetFruit() = 0;
+	virtual TSharedRef<IBerry, Mode> GetBerry() = 0;
+	virtual TSharedRef<IFruit, Mode> GetFruit() = 0;
 };
 
-struct FSmoothie : public ISmoothie
+template<ESPMode Mode = ESPMode::Fast>
+struct TSmoothie : public ISmoothie<Mode>
 {
-	FSmoothie(TSharedRef<IFruit> InFruit, TSharedRef<IBerry> InBerry) : Berry(InBerry), Fruit(InFruit) { }
-	virtual ~FSmoothie() { }
-	virtual TSharedRef<IBerry> GetBerry() override { return Berry; }
-	virtual TSharedRef<IFruit> GetFruit() override { return Fruit; }
-	TSharedRef<IBerry> Berry;
-	TSharedRef<IFruit> Fruit;
+	TSmoothie(TSharedRef<IFruit, Mode> InFruit, TSharedRef<IBerry, Mode> InBerry) : Berry(InBerry), Fruit(InFruit) { }
+	virtual ~TSmoothie() { }
+	virtual TSharedRef<IBerry, Mode> GetBerry() override { return Berry; }
+	virtual TSharedRef<IFruit, Mode> GetFruit() override { return Fruit; }
+	TSharedRef<IBerry, Mode> Berry;
+	TSharedRef<IFruit, Mode> Fruit;
 };
 
-struct FTwoSmoothies
+template<ESPMode Mode = ESPMode::Fast>
+struct TTwoSmoothies
 {
-	TSharedPtr<ISmoothie> One;
-	TSharedPtr<ISmoothie> Two;
+	TSharedPtr<ISmoothie<Mode>, Mode> One;
+	TSharedPtr<ISmoothie<Mode>, Mode> Two;
 };
 
 
 DECLARE_DELEGATE_RetVal(TSharedRef<IBerry>, FBerryFactoryDelegate);
 DECLARE_DELEGATE_RetVal(TSharedRef<IFruit>, FFruitFactoryDelegate);
-DECLARE_DELEGATE_RetVal_TwoParams(TSharedRef<ISmoothie>, FSmoothieFactoryDelegate, TSharedRef<IFruit>, TSharedRef<IBerry>);
+DECLARE_DELEGATE_RetVal_TwoParams(TSharedRef<ISmoothie<>>, FSmoothieFactoryDelegate, TSharedRef<IFruit>, TSharedRef<IBerry>);
 
 Expose_TNameOf(IBerry)
 Expose_TNameOf(IFruit)
-Expose_TNameOf(ISmoothie)
+Expose_TNameOf(ISmoothie<>)
+Expose_TNameOf(ISmoothie<ESPMode::ThreadSafe>)
 
 
 /* Tests
@@ -78,15 +82,12 @@ Expose_TNameOf(ISmoothie)
 
 bool FTypeContainerTest::RunTest(const FString& Parameters)
 {
-	// gmp: temporarily disabling unit test
-	return true;
-
 	// existing instance test
 	{
 		TSharedRef<IFruit> Fruit = MakeShareable(new FBanana());
 		TSharedRef<IBerry> Berry = MakeShareable(new FStrawberry());
 
-		FTypeContainer Container;
+		TTypeContainer<> Container;
 		{
 			Container.RegisterInstance<IFruit>(Fruit);
 			Container.RegisterInstance<IBerry>(Berry);
@@ -100,15 +101,15 @@ bool FTypeContainerTest::RunTest(const FString& Parameters)
 
 	// per instance test
 	{
-		FTypeContainer Container;
+		TTypeContainer<> Container;
 		{
 			Container.RegisterClass<IFruit, FBanana>(ETypeContainerScope::Instance);
 			Container.RegisterClass<IBerry, FStrawberry>(ETypeContainerScope::Instance);
-			Container.RegisterClass<ISmoothie, FSmoothie, IFruit, IBerry>(ETypeContainerScope::Instance); // !!!
+			Container.RegisterClass<ISmoothie<>, TSmoothie<>, IFruit, IBerry>(ETypeContainerScope::Instance); // !!!
 		}
 
-		auto Smoothie1 = Container.GetInstance<ISmoothie>();
-		auto Smoothie2 = Container.GetInstance<ISmoothie>();
+		auto Smoothie1 = Container.GetInstance<ISmoothie<>>();
+		auto Smoothie2 = Container.GetInstance<ISmoothie<>>();
 
 		TestNotEqual(TEXT("For per-instances classes, a unique instance must be returned each time"), Smoothie1, Smoothie2);
 		TestNotEqual(TEXT("For per-instances dependencies, a unique instance must be returned each time [1]"), Smoothie1->GetBerry(), Smoothie2->GetBerry());
@@ -117,28 +118,28 @@ bool FTypeContainerTest::RunTest(const FString& Parameters)
 
 	// per thread test
 	{
-		FTypeContainer Container;
+		TTypeContainer<ESPMode::ThreadSafe> Container;
 		{
 			Container.RegisterClass<IFruit, FBanana>(ETypeContainerScope::Thread);
 			Container.RegisterClass<IBerry, FStrawberry>(ETypeContainerScope::Instance);
-			Container.RegisterClass<ISmoothie, FSmoothie, IFruit, IBerry>(ETypeContainerScope::Thread); // !!!
+			Container.RegisterClass<ISmoothie<ESPMode::ThreadSafe>, TSmoothie<ESPMode::ThreadSafe>, IFruit, IBerry>(ETypeContainerScope::Thread); // !!!
 		}
 
-		TFunction<FTwoSmoothies()> MakeSmoothies = [&]()
+		TFunction<TTwoSmoothies<ESPMode::ThreadSafe>()> MakeSmoothies = [&]()
 		{
-			FTwoSmoothies Smoothies;
-			Smoothies.One = Container.GetInstance<ISmoothie>();
-			Smoothies.Two = Container.GetInstance<ISmoothie>();
+			TTwoSmoothies<ESPMode::ThreadSafe> Smoothies;
+			Smoothies.One = Container.GetInstance<ISmoothie<ESPMode::ThreadSafe>>();
+			Smoothies.Two = Container.GetInstance<ISmoothie<ESPMode::ThreadSafe>>();
 			return Smoothies;
 		};
 
 		auto Smoothies1 = Async(EAsyncExecution::Thread, MakeSmoothies);
 		auto Smoothies2 = Async(EAsyncExecution::Thread, MakeSmoothies);
 
-		TSharedPtr<ISmoothie> One1 = Smoothies1.Get().One;
-		TSharedPtr<ISmoothie> Two1 = Smoothies1.Get().Two;
-		TSharedPtr<ISmoothie> One2 = Smoothies2.Get().One;
-		TSharedPtr<ISmoothie> Two2 = Smoothies2.Get().Two;
+		auto One1 = Smoothies1.Get().One;
+		auto Two1 = Smoothies1.Get().Two;
+		auto One2 = Smoothies2.Get().One;
+		auto Two2 = Smoothies2.Get().Two;
 
 		TestEqual(TEXT("For per-thread classes, the same instance must be returned from the same thread [1]"), One1, Two1);
 		TestEqual(TEXT("For per-thread classes, the same instance must be returned from the same thread [2]"), One2, Two2);
@@ -148,28 +149,28 @@ bool FTypeContainerTest::RunTest(const FString& Parameters)
 
 	// per process test
 	{
-		FTypeContainer Container;
+		TTypeContainer<ESPMode::ThreadSafe> Container;
 		{
 			Container.RegisterClass<IFruit, FBanana>(ETypeContainerScope::Thread);
 			Container.RegisterClass<IBerry, FStrawberry>(ETypeContainerScope::Instance);
-			Container.RegisterClass<ISmoothie, FSmoothie, IFruit, IBerry>(ETypeContainerScope::Process); //!!!
+			Container.RegisterClass<ISmoothie<ESPMode::ThreadSafe>, TSmoothie<ESPMode::ThreadSafe>, IFruit, IBerry>(ETypeContainerScope::Process); //!!!
 		}
 
-		TFunction<FTwoSmoothies()> MakeSmoothies = [&]()
+		TFunction<TTwoSmoothies<ESPMode::ThreadSafe>()> MakeSmoothies = [&]()
 		{
-			FTwoSmoothies Smoothies;
-			Smoothies.One = Container.GetInstance<ISmoothie>();
-			Smoothies.Two = Container.GetInstance<ISmoothie>();
+			TTwoSmoothies<ESPMode::ThreadSafe> Smoothies;
+			Smoothies.One = Container.GetInstance<ISmoothie<ESPMode::ThreadSafe>>();
+			Smoothies.Two = Container.GetInstance<ISmoothie<ESPMode::ThreadSafe>>();
 			return Smoothies;
 		};
 
 		auto Smoothies1 = Async(EAsyncExecution::Thread, MakeSmoothies);
 		auto Smoothies2 = Async(EAsyncExecution::Thread, MakeSmoothies);
 
-		TSharedPtr<ISmoothie> One1 = Smoothies1.Get().One;
-		TSharedPtr<ISmoothie> Two1 = Smoothies1.Get().Two;
-		TSharedPtr<ISmoothie> One2 = Smoothies2.Get().One;
-		TSharedPtr<ISmoothie> Two2 = Smoothies2.Get().Two;
+		auto One1 = Smoothies1.Get().One;
+		auto Two1 = Smoothies1.Get().Two;
+		auto One2 = Smoothies2.Get().One;
+		auto Two2 = Smoothies2.Get().Two;
 
 		TestEqual(TEXT("For per-process classes, the same instance must be returned from the same thread [1]"), One1, Two1);
 		TestEqual(TEXT("For per-process classes, the same instance must be returned from the same thread [2]"), One2, Two2);
@@ -186,13 +187,13 @@ bool FTypeContainerTest::RunTest(const FString& Parameters)
 				return MakeShareable(new FStrawberry());
 			}
 
-			static TSharedRef<ISmoothie> MakeSmoothie(TSharedRef<IFruit> Fruit, TSharedRef<IBerry> Berry)
+			static TSharedRef<ISmoothie<>> MakeSmoothie(TSharedRef<IFruit> Fruit, TSharedRef<IBerry> Berry)
 			{
-				return MakeShareable(new FSmoothie(Fruit, Berry));
+				return MakeShareable(new TSmoothie<>(Fruit, Berry));
 			}
 		};
 
-		FTypeContainer Container;
+		TTypeContainer<> Container;
 		{
 			Container.RegisterFactory<IBerry>(&FLocal::MakeStrawberry);
 			Container.RegisterFactory<IFruit>(
@@ -200,12 +201,12 @@ bool FTypeContainerTest::RunTest(const FString& Parameters)
 					return MakeShareable(new FBanana());
 				}
 			);
-			Container.RegisterFactory<ISmoothie, IFruit, IBerry>(&FLocal::MakeSmoothie);
+			Container.RegisterFactory<ISmoothie<>, IFruit, IBerry>(&FLocal::MakeSmoothie);
 		}
 
 		auto Berry = Container.GetInstance<IBerry>();
 		auto Fruit = Container.GetInstance<IFruit>();
-		auto Smoothie = Container.GetInstance<ISmoothie>();
+		auto Smoothie = Container.GetInstance<ISmoothie<>>();
 	}
 
 	// delegate test
@@ -227,21 +228,21 @@ bool FTypeContainerTest::RunTest(const FString& Parameters)
 				return MakeShareable(new FStrawberry());
 			}
 
-			static TSharedRef<ISmoothie> MakeSmoothie(TSharedRef<IFruit> Fruit, TSharedRef<IBerry> Berry)
+			static TSharedRef<ISmoothie<>> MakeSmoothie(TSharedRef<IFruit> Fruit, TSharedRef<IBerry> Berry)
 			{
-				return MakeShareable(new FSmoothie(Fruit, Berry));
+				return MakeShareable(new TSmoothie<>(Fruit, Berry));
 			}
 		};
 
-		FTypeContainer Container;
+		TTypeContainer<> Container;
 		{
 			Container.RegisterDelegate<IBerry>(FBerryFactoryDelegate::CreateStatic(&FLocal::MakeBerry));
 			Container.RegisterDelegate<IFruit>(FFruitFactoryDelegate::CreateStatic(&FLocal::MakeFruit, true));
-			Container.RegisterDelegate<ISmoothie, FSmoothieFactoryDelegate, IFruit, IBerry>(FSmoothieFactoryDelegate::CreateStatic(&FLocal::MakeSmoothie));
+			Container.RegisterDelegate<ISmoothie<>, FSmoothieFactoryDelegate, IFruit, IBerry>(FSmoothieFactoryDelegate::CreateStatic(&FLocal::MakeSmoothie));
 		}
 
 		auto Fruit = Container.GetInstance<IFruit>();
-		auto Smoothie = Container.GetInstance<ISmoothie>();
+		auto Smoothie = Container.GetInstance<ISmoothie<>>();
 	}
 
 	return true;
diff --git a/Engine/Source/Runtime/Core/Public/Async/Async.h b/Engine/Source/Runtime/Core/Public/Async/Async.h
--- a/Engine/Source/Runtime/Core/Public/Async/Async.h
+++ b/Engine/Source/Runtime/Core/Public/Async/Async.h
@@ -26,13 +26,13 @@ enum class EAsyncExecution
  * Template for setting a promise's value from a function.
  */
 template<typename ResultType>
-inline void SetPromise(TPromise<ResultType>& Promise, TFunction<ResultType()> Function)
+inline void SetPromise(TPromise<ResultType>& Promise, const TFunction<ResultType()>& Function)
 {
 	Promise.SetValue(Function());
 }
 
 template<typename ResultType>
-inline void SetPromise(TPromise<ResultType>& Promise, TFunctionRef<ResultType()> Function)
+inline void SetPromise(TPromise<ResultType>& Promise, const TFunctionRef<ResultType()>& Function)
 {
 	Promise.SetValue(Function());
 }
@@ -42,14 +42,14 @@ inline void SetPromise(TPromise<ResultType>& Promise, TFunctionRef<ResultType()>
  * Template for setting a promise's value from a function (specialization for void results).
  */
 template<>
-inline void SetPromise(TPromise<void>& Promise, TFunction<void()> Function)
+inline void SetPromise(TPromise<void>& Promise, const TFunction<void()>& Function)
 {
 	Function();
 	Promise.SetValue();
 }
 
 template<>
-inline void SetPromise(TPromise<void>& Promise, TFunctionRef<void()> Function)
+inline void SetPromise(TPromise<void>& Promise, const TFunctionRef<void()>& Function)
 {
 	Function();
 	Promise.SetValue();
@@ -165,7 +165,7 @@ public:
 	 * @param InThreadFuture The thread that is running this task.
 	 */
 	TAsyncRunnable(TFunction<ResultType()>&& InFunction, TPromise<ResultType>&& InPromise, TFuture<FRunnableThread*>&& InThreadFuture)
-		: Function(InFunction)
+		: Function(MoveTemp(InFunction))
 		, Promise(MoveTemp(InPromise))
 		, ThreadFuture(MoveTemp(InThreadFuture))
 	{ }
@@ -205,7 +205,7 @@ public:
 	 * @param InPromise The promise object used to return the function's result.
 	 */
 	TAsyncQueuedWork(TFunction<ResultType()>&& InFunction, TPromise<ResultType>&& InPromise)
-		: Function(InFunction)
+		: Function(MoveTemp(InFunction))
 		, Promise(MoveTemp(InPromise))
 	{ }
 
@@ -284,12 +284,13 @@ struct FAsyncThreadIndex
  * @param ResultType The type of the function's return value.
  * @param Execution The execution method to use, i.e. on Task Graph or in a separate thread.
  * @param Function The function to execute.
+ * @param CompletionCallback An optional callback function that is executed when the function completed execution.
  * @result A TFuture object that will receive the return value from the function.
  */
 template<typename ResultType>
-TFuture<ResultType> Async(EAsyncExecution Execution, TFunction<ResultType()> Function)
+TFuture<ResultType> Async(EAsyncExecution Execution, TFunction<ResultType()> Function, TFunction<void()> CompletionCallback = TFunction<void()>())
 {
-	TPromise<ResultType> Promise;
+	TPromise<ResultType> Promise(MoveTemp(CompletionCallback));
 	TFuture<ResultType> Future = Promise.GetFuture();
 
 	switch (Execution)
@@ -344,7 +345,6 @@ template<typename ResultType>
 uint32 TAsyncRunnable<ResultType>::Run()
 {
 	SetPromise(Promise, Function);
-
 	FRunnableThread* Thread = ThreadFuture.Get();
 
 	// Enqueue deletion of the thread to a different thread.
diff --git a/Engine/Source/Runtime/Core/Public/Async/Future.h b/Engine/Source/Runtime/Core/Public/Async/Future.h
--- a/Engine/Source/Runtime/Core/Public/Async/Future.h
+++ b/Engine/Source/Runtime/Core/Public/Async/Future.h
@@ -19,6 +19,17 @@ public:
 		, Complete(false)
 	{ }
 
+	/**
+	 * Create and initialize a new instance with a callback.
+	 *
+	 * @param InCompletionCallback A function that is called when the state is completed.
+	 */
+	FFutureState(TFunction<void()>&& InCompletionCallback)
+		: CompletionCallback(InCompletionCallback)
+		, CompletionEvent(FPlatformProcess::GetSynchEventFromPool(true))
+		, Complete(false)
+	{ }
+
 	/** Destructor. */
 	~FFutureState()
 	{
@@ -62,11 +73,20 @@ protected:
 	void MarkComplete()
 	{
 		Complete = true;
+
+		if (CompletionCallback)
+		{
+			CompletionCallback();
+		}
+
 		CompletionEvent->Trigger();
 	}
 
 private:
 
+	/** An optional callback function that is executed the state is completed. */
+	TFunction<void()> CompletionCallback;
+
 	/** Holds an event signaling that the result is available. */
 	FEvent* CompletionEvent;
 
@@ -82,6 +102,22 @@ template<typename InternalResultType>
 class TFutureState
 	: public FFutureState
 {
+public:
+
+	/** Default constructor. */
+	TFutureState()
+		: FFutureState()
+	{ }
+
+	/**
+	 * Create and initialize a new instance with a callback.
+	 *
+	 * @param CompletionCallback A function that is called when the state is completed.
+	 */
+	TFutureState(TFunction<void()>&& CompletionCallback)
+		: FFutureState(MoveTemp(CompletionCallback))
+	{ }
+
 public:
 
 	/**
@@ -253,7 +289,7 @@ protected:
 	{
 		// if you hit this assertion then your future has an invalid state.
 		// this happens if you have an uninitialized future or if you moved
-		// the future it to another instance.
+		// it to another instance.
 		check(State.IsValid());
 
 		return State;
@@ -746,6 +782,15 @@ public:
 		Other.State.Reset();
 	}
 
+	/**
+	 * Create and initialize a new instance with a callback.
+	 *
+	 * @param CompletionCallback A function that is called when the future state is completed.
+	 */
+	TPromiseBase(TFunction<void()>&& CompletionCallback)
+		: State(MakeShareable(new TFutureState<InternalResultType>(MoveTemp(CompletionCallback))))
+	{ }
+
 public:
 
 	/** Move assignment operator. */
@@ -817,6 +862,16 @@ public:
 		, FutureRetrieved(MoveTemp(Other.FutureRetrieved))
 	{ }
 
+	/**
+	 * Create and initialize a new instance with a callback.
+	 *
+	 * @param CompletionCallback A function that is called when the future state is completed.
+	 */
+	TPromise(TFunction<void()>&& CompletionCallback)
+		: BaseType(MoveTemp(CompletionCallback))
+		, FutureRetrieved(false)
+	{ }
+
 public:
 
 	/**
@@ -907,6 +962,16 @@ public:
 		, FutureRetrieved(MoveTemp(Other.FutureRetrieved))
 	{ }
 
+	/**
+	 * Create and initialize a new instance with a callback.
+	 *
+	 * @param CompletionCallback A function that is called when the future state is completed.
+	 */
+	TPromise(TFunction<void()>&& CompletionCallback)
+		: BaseType(MoveTemp(CompletionCallback))
+		, FutureRetrieved(false)
+	{ }
+
 public:
 
 	/**
@@ -984,6 +1049,16 @@ public:
 		, FutureRetrieved(false)
 	{ }
 
+	/**
+	 * Create and initialize a new instance with a callback.
+	 *
+	 * @param CompletionCallback A function that is called when the future state is completed.
+	 */
+	TPromise(TFunction<void()>&& CompletionCallback)
+		: BaseType(MoveTemp(CompletionCallback))
+		, FutureRetrieved(false)
+	{ }
+
 public:
 
 	/**
diff --git a/Engine/Source/Runtime/Core/Public/Containers/Queue.h b/Engine/Source/Runtime/Core/Public/Containers/Queue.h
--- a/Engine/Source/Runtime/Core/Public/Containers/Queue.h
+++ b/Engine/Source/Runtime/Core/Public/Containers/Queue.h
@@ -78,15 +78,20 @@ public:
 		OutItem = Popped->Item;
 
 		TNode* OldTail = Tail;
-
 		Tail = Popped;
 		Tail->Item = ItemType();
-
 		delete OldTail;
 
 		return true;
 	}
 
+	/** Empty the queue, discarding all items. */
+	void Empty()
+	{
+		ItemType DummyItem;
+		while (Dequeue(DummyItem));
+	}
+
 	/**
 	 * Adds an item to the head of the queue.
 	 *
diff --git a/Engine/Source/Runtime/Core/Public/Misc/TypeContainer.h b/Engine/Source/Runtime/Core/Public/Misc/TypeContainer.h
--- a/Engine/Source/Runtime/Core/Public/Misc/TypeContainer.h
+++ b/Engine/Source/Runtime/Core/Public/Misc/TypeContainer.h
@@ -22,7 +22,7 @@ enum class ETypeContainerScope
 
 
 /**
- * Implements a type container.
+ * Template for type containers.
  *
  * Type containers allow for configuring objects and their dependencies without actually
  * creating these objects. They fulfill a role similar to class factories, but without
@@ -46,10 +46,16 @@ enum class ETypeContainerScope
  * See the file TypeContainerTest.cpp for detailed examples on how to use each of these
  * type registration methods.
  *
+ * In the interest of fast performance, the object pointers returned by type containers are
+ * not thread-safe by default. If you intend to use the same type container and share its
+ * objects from different threads, use TTypeContainer<ESPMode::ThreadSafe> instead, which
+ * will then manage and return thread-safe versions of all object pointers.
+ *
  * Note: Type containers depend on variadic templates and are therefore not available
  * on XboxOne at this time, which means they should only be used for desktop applications.
  */
-class FTypeContainer
+template<ESPMode Mode = ESPMode::Fast>
+class TTypeContainer
 {
 	/** Interface for object instance providers. */
 	class IInstanceProvider
@@ -62,9 +68,9 @@ class FTypeContainer
 		/**
 		 * Gets an instance of a class.
 		 *
-		 * @return Shared pointer to the instance (must be cast to TSharedPtr<R>).
+		 * @return Shared pointer to the instance (must be cast to TSharedPtr<R, Mode>).
 		 */
-		virtual TSharedPtr<void> GetInstance() = 0;
+		virtual TSharedPtr<void, Mode> GetInstance() = 0;
 	};
 
 	/** Implements an instance provider that forwards instance requests to a factory function. */
@@ -72,18 +78,18 @@ class FTypeContainer
 	struct TFunctionInstanceProvider
 		: public IInstanceProvider
 	{
-		TFunctionInstanceProvider(TFunction<TSharedPtr<void>()>&& InCreateFunc)
+		TFunctionInstanceProvider(TFunction<TSharedPtr<void, Mode>()>&& InCreateFunc)
 			: CreateFunc(MoveTemp(InCreateFunc))
 		{ }
 
 		virtual ~TFunctionInstanceProvider() override { }
 
-		virtual TSharedPtr<void> GetInstance() override
+		virtual TSharedPtr<void, Mode> GetInstance() override
 		{
 			return CreateFunc();
 		}
 
-		TFunction<TSharedPtr<void>()> CreateFunc;
+		TFunction<TSharedPtr<void, Mode>()> CreateFunc;
 	};
 
 
@@ -92,18 +98,18 @@ class FTypeContainer
 	struct TSharedInstanceProvider
 		: public IInstanceProvider
 	{
-		TSharedInstanceProvider(const TSharedRef<T>& InInstance)
+		TSharedInstanceProvider(const TSharedRef<T, Mode>& InInstance)
 			: Instance(InInstance)
 		{ }
 
 		virtual ~TSharedInstanceProvider() { }
 
-		virtual TSharedPtr<void> GetInstance() override
+		virtual TSharedPtr<void, Mode> GetInstance() override
 		{
 			return Instance;
 		}
 
-		TSharedRef<T> Instance;
+		TSharedRef<T, Mode> Instance;
 	};
 
 
@@ -111,7 +117,7 @@ class FTypeContainer
 	struct FThreadInstanceProvider
 		: public IInstanceProvider
 	{
-		FThreadInstanceProvider(TFunction<TSharedPtr<void>()>&& InCreateFunc)
+		FThreadInstanceProvider(TFunction<TSharedPtr<void, Mode>()>&& InCreateFunc)
 			: CreateFunc(MoveTemp(InCreateFunc))
 			, TlsSlot(FPlatformTLS::AllocTlsSlot())
 		{ }
@@ -121,7 +127,7 @@ class FTypeContainer
 			FPlatformTLS::FreeTlsSlot(TlsSlot);
 		}
 
-		TFunction<TSharedPtr<void>()> CreateFunc;
+		TFunction<TSharedPtr<void, Mode>()> CreateFunc;
 		uint32 TlsSlot;
 	};
 
@@ -131,23 +137,23 @@ class FTypeContainer
 	struct TThreadInstanceProvider
 		: public FThreadInstanceProvider
 	{
-		typedef TTlsAutoCleanupValue<TSharedPtr<T>> TlsValueType;
+		typedef TTlsAutoCleanupValue<TSharedPtr<T, Mode>> TlsValueType;
 
-		TThreadInstanceProvider(TFunction<TSharedPtr<void>()>&& InCreateFunc)
+		TThreadInstanceProvider(TFunction<TSharedPtr<void, Mode>()>&& InCreateFunc)
 			: FThreadInstanceProvider(MoveTemp(InCreateFunc))
 		{ }
 
 		virtual ~TThreadInstanceProvider() override { }
 
-		virtual TSharedPtr<void> GetInstance() override
+		virtual TSharedPtr<void, Mode> GetInstance() override
 		{
-			auto TlsValue = (TlsValueType*)FPlatformTLS::GetTlsValue(TlsSlot);
+			auto TlsValue = (TlsValueType*)FPlatformTLS::GetTlsValue(this->TlsSlot);
 
 			if (TlsValue == nullptr)
 			{
-				TlsValue = new TlsValueType(StaticCastSharedPtr<T>(CreateFunc()));
+				TlsValue = new TlsValueType(StaticCastSharedPtr<T>(this->CreateFunc()));
 				TlsValue->Register();
-				FPlatformTLS::SetTlsValue(TlsSlot, TlsValue);
+				FPlatformTLS::SetTlsValue(this->TlsSlot, TlsValue);
 			}
 
 			return TlsValue->Get();
@@ -164,7 +170,7 @@ public:
 	 * @see RegisterClass, RegisterDelegate, RegisterFactory, RegisterInstance
 	 */
 	template<class R>
-	TSharedRef<R> GetInstance() const
+	TSharedRef<R, Mode> GetInstance() const
 	{
 		FScopeLock Lock(&CriticalSection);
 		{
@@ -186,7 +192,7 @@ public:
 	 * @see GetInstance, RegisterClass, RegisterInstance
 	 */
 	template<class R>
-	TSharedRef<R> GetInstanceRef() const
+	TSharedRef<R, Mode> GetInstanceRef() const
 	{
 		return GetInstance<R>().ToSharedRef();
 	}
@@ -227,7 +233,7 @@ public:
 		case ETypeContainerScope::Instance:
 			Provider = MakeShareable(
 				new TFunctionInstanceProvider<T>(
-					[this]() -> TSharedPtr<void> {
+					[this]() -> TSharedPtr<void, Mode> {
 						return MakeShareable(new T(GetInstance<P>()...));
 					}
 				)
@@ -237,7 +243,7 @@ public:
 		case ETypeContainerScope::Thread:
 			Provider = MakeShareable(
 				new TThreadInstanceProvider<T>(
-					[this]() -> TSharedPtr<void> {
+					[this]() -> TSharedPtr<void, Mode> {
 						return MakeShareable(new T(GetInstance<P>()...));
 					}
 				)
@@ -267,11 +273,11 @@ public:
 	template<class R, class D, typename... P>
 	void RegisterDelegate(D Delegate)
 	{
-		static_assert(TAreTypesEqual<TSharedRef<R>, typename D::RetValType>::Value, "Delegate return type must be TSharedPtr<R>");
+		static_assert(TAreTypesEqual<TSharedRef<R, Mode>, typename D::RetValType>::Value, "Delegate return type must be TSharedPtr<R>");
 
 		TSharedPtr<IInstanceProvider> Provider = MakeShareable(
 			new TFunctionInstanceProvider<R>(
-				[=]() -> TSharedPtr<void> {
+				[=]() -> TSharedPtr<void, Mode> {
 					return Delegate.Execute(GetInstance<P>()...);
 				}
 			)
@@ -289,11 +295,11 @@ public:
 	 * @see RegisterClass, RegisterInstance, Unregister
 	 */
 	template<class R>
-	void RegisterFactory(const TFunction<TSharedRef<R>()>& CreateFunc)
+	void RegisterFactory(const TFunction<TSharedRef<R, Mode>()>& CreateFunc)
 	{
 		TSharedPtr<IInstanceProvider> Provider = MakeShareable(
 			new TFunctionInstanceProvider<R>(
-				[=]() -> TSharedPtr<void> {
+				[=]() -> TSharedPtr<void, Mode> {
 					return CreateFunc();
 				}
 			)
@@ -312,11 +318,11 @@ public:
 	 * @see RegisterClass, RegisterInstance, Unregister
 	 */
 	template<class R, typename P0, typename... P>
-	void RegisterFactory(const TFunction<TSharedRef<R>(TSharedRef<P0>, TSharedRef<P>...)>& CreateFunc)
+	void RegisterFactory(const TFunction<TSharedRef<R, Mode>(TSharedRef<P0, Mode>, TSharedRef<P, Mode>...)>& CreateFunc)
 	{
 		TSharedPtr<IInstanceProvider> Provider = MakeShareable(
 			new TFunctionInstanceProvider<R>(
-				[=]() -> TSharedPtr<void> {
+				[=]() -> TSharedPtr<void, Mode> {
 					return CreateFunc(GetInstance<P0>(), GetInstance<P>()...);
 				}
 			)
@@ -337,11 +343,11 @@ public:
 	 * @see RegisterClass, RegisterInstance, Unregister
 	 */
 	template<class R, typename... P>
-	void RegisterFactory(TSharedRef<R> (*CreateFunc)(TSharedRef<P>...))
+	void RegisterFactory(TSharedRef<R, Mode> (*CreateFunc)(TSharedRef<P, Mode>...))
 	{
 		TSharedPtr<IInstanceProvider> Provider = MakeShareable(
 			new TFunctionInstanceProvider<R>(
-				[=]() -> TSharedPtr<void> {
+				[=]() -> TSharedPtr<void, Mode> {
 					return CreateFunc(GetInstance<P>()...);
 				}
 			)
@@ -360,7 +366,7 @@ public:
 	 * @see RegisterClass, RegisterDelegate, Unregister
 	 */
 	template<class R, class T>
-	void RegisterInstance(const TSharedRef<T>& Instance)
+	void RegisterInstance(const TSharedRef<T, Mode>& Instance)
 	{
 		static_assert(TPointerIsConvertibleFromTo<T, R>::Value, "T must implement R");
 
@@ -408,4 +414,8 @@ private:
 };
 
 
+/** Thread-unsafe type container (for backwards compatibility). */
+class FTypeContainer : public TTypeContainer<ESPMode::Fast> { };
+
+
 #endif //PLATFORM_COMPILER_HAS_VARIADIC_TEMPLATES
diff --git a/Engine/Source/Runtime/CoreUObject/Classes/Object.h b/Engine/Source/Runtime/CoreUObject/Classes/Object.h
--- a/Engine/Source/Runtime/CoreUObject/Classes/Object.h
+++ b/Engine/Source/Runtime/CoreUObject/Classes/Object.h
@@ -270,6 +270,10 @@ enum class EUnit
 	
 	Milliseconds, Seconds, Minutes, Hours, Days, Months, Years,
 
+	/** Arbitrary multiplier */
+	
+	Multiplier,
+
 
 	/** Symbolic entry, not specifiable on meta data. */
 	
diff --git a/Engine/Source/Runtime/Engine/Classes/Engine/Engine.h b/Engine/Source/Runtime/Engine/Classes/Engine/Engine.h
--- a/Engine/Source/Runtime/Engine/Classes/Engine/Engine.h
+++ b/Engine/Source/Runtime/Engine/Classes/Engine/Engine.h
@@ -1660,7 +1660,6 @@ public:
 	bool HandleFlushLogCommand( const TCHAR* Cmd, FOutputDevice& Ar );
 	bool HandleGameVerCommand( const TCHAR* Cmd, FOutputDevice& Ar );
 	bool HandleStatCommand( UWorld* World, FCommonViewportClient* ViewportClient, const TCHAR* Cmd, FOutputDevice& Ar );
-	bool HandleStartMovieCaptureCommand( const TCHAR* Cmd, FOutputDevice& Ar );
 	bool HandleStopMovieCaptureCommand( const TCHAR* Cmd, FOutputDevice& Ar );
 	bool HandleCrackURLCommand( const TCHAR* Cmd, FOutputDevice& Ar );
 	bool HandleDeferCommand( const TCHAR* Cmd, FOutputDevice& Ar );
diff --git a/Engine/Source/Runtime/Engine/Classes/Engine/GameInstance.h b/Engine/Source/Runtime/Engine/Classes/Engine/GameInstance.h
--- a/Engine/Source/Runtime/Engine/Classes/Engine/GameInstance.h
+++ b/Engine/Source/Runtime/Engine/Classes/Engine/GameInstance.h
@@ -96,9 +96,6 @@ public:
 
 	virtual bool StartPIEGameInstance(ULocalPlayer* LocalPlayer, bool bInSimulateInEditor, bool bAnyBlueprintErrors, bool bStartInSpectatorMode);
 
-	/** Initial map to load when starting the game instance. Overrides anything set on the command line */
-	FString InitialMapOverride;
-
 #endif
 
 	class UEngine* GetEngine() const;
diff --git a/Engine/Source/Runtime/Engine/Private/AVIWriter.cpp b/Engine/Source/Runtime/Engine/Private/AVIWriter.cpp
--- a/Engine/Source/Runtime/Engine/Private/AVIWriter.cpp
+++ b/Engine/Source/Runtime/Engine/Private/AVIWriter.cpp
@@ -25,7 +25,6 @@ typedef TCHAR* PTCHAR;
 
 #include "CapturePin.h"
 #include "CaptureSource.h"
-#include "SlateCore.h"
 
 
 #define g_wszCapture    L"Capture Filter"
@@ -178,22 +177,8 @@ public:
 
 public:
 
-	virtual void StartCapture(TWeakPtr<FSceneViewport> InViewport) override
+	virtual void Initialize() override
 	{
-		auto Viewport = InViewport.Pin();
-		if (bCapturing || !Viewport.IsValid())
-		{
-			return;
-		}
-
-		if ( Viewport->GetSize().GetMin() <= 0)
-		{
-			UE_LOG(LogAVIWriter, Error, TEXT( "ERROR - Attempting to record a 0 sized viewport!" ));
-			return;
-		}
-
-		CaptureViewport = InViewport;
-
 		// Initialize the COM library.
 		if (!FWindowsPlatformMisc::CoInitialize()) 
 		{
@@ -306,12 +291,34 @@ public:
 		// Now connect the graph
 		if (EncodingFilter)
 		{
-			Graph->Connect(GetPin(CaptureFilter, PINDIR_OUTPUT), GetPin(EncodingFilter, PINDIR_INPUT));
-			Graph->Connect(GetPin(EncodingFilter, PINDIR_OUTPUT), GetPin(pMux, PINDIR_INPUT));
+			hr = Graph->Connect(GetPin(CaptureFilter, PINDIR_OUTPUT), GetPin(EncodingFilter, PINDIR_INPUT));
+			if (FAILED(hr))
+			{
+				UE_LOG(LogAVIWriter, Error, TEXT( "ERROR - Failed to connect capture filter to encoding filter! (%d)" ), hr);
+				Graph->Release();
+				FWindowsPlatformMisc::CoUninitialize(); 
+				return;
+			}
+
+			hr = Graph->Connect(GetPin(EncodingFilter, PINDIR_OUTPUT), GetPin(pMux, PINDIR_INPUT));
+			if (FAILED(hr))
+			{
+				UE_LOG(LogAVIWriter, Error, TEXT( "ERROR - Failed to connect encoding filter to muxer! (%d)" ), hr);
+				Graph->Release();
+				FWindowsPlatformMisc::CoUninitialize(); 
+				return;
+			}
 		}
 		else
 		{
-			Graph->Connect(GetPin(CaptureFilter, PINDIR_OUTPUT), GetPin(pMux, PINDIR_INPUT));
+			hr = Graph->Connect(GetPin(CaptureFilter, PINDIR_OUTPUT), GetPin(pMux, PINDIR_INPUT));
+			if (FAILED(hr))
+			{
+				UE_LOG(LogAVIWriter, Error, TEXT( "ERROR - Failed to connect capture filter to muxer! (%d)" ), hr);
+				Graph->Release();
+				FWindowsPlatformMisc::CoUninitialize(); 
+				return;
+			}
 		}
 
 		Graph->Connect(GetPin(pMux, PINDIR_OUTPUT), GetPin(FileWriter, PINDIR_INPUT));
@@ -332,7 +339,7 @@ public:
 		pMux->Release();
 	}
 
-	void StopCapture()
+	void Finalize()
 	{
 		if (!bCapturing)
 		{
@@ -344,7 +351,6 @@ public:
 		Control->Stop();
 
 		bCapturing = false;
-		CaptureViewport = nullptr;
 		FrameNumber = 0;
 
 		SAFE_RELEASE(EncodingFilter);
@@ -355,11 +361,6 @@ public:
 		FWindowsPlatformMisc::CoUninitialize();
 	}
 
-	void Close()
-	{
-		StopCapture();
-	}
-
 	virtual void DropFrames(int32 NumFramesToDrop) override
 	{
 		FrameNumber += NumFramesToDrop;
@@ -394,22 +395,12 @@ public:
 	
 public:
 	
-	virtual void StartCapture(TWeakPtr<FSceneViewport> InViewport) override
+	virtual void Initialize() override
 	{
 		SCOPED_AUTORELEASE_POOL;
 
-		auto Viewport = InViewport.Pin();
-
-		if (!bCapturing && Viewport.IsValid())
+		if (!bCapturing)
 		{
-			if ( Viewport->GetSize().GetMin() <= 0)
-			{
-				UE_LOG(LogAVIWriter, Error, TEXT( "ERROR - Attempting to record a 0 sized viewport!" ));
-				return;
-			}
-			
-			CaptureViewport = InViewport;
-			
 			// Attempt to make the dir if it doesn't exist.
 			TCHAR File[MAX_SPRINTF] = TEXT("");
 			IFileManager::Get().MakeDirectory(*FPaths::GetPath(Options.OutputFilename), true);
@@ -436,16 +427,16 @@ public:
 				VideoSettings = [NSDictionary dictionaryWithObjectsAndKeys:
 								 AVVideoCodecJPEG, AVVideoCodecKey,
 								 [NSDictionary dictionaryWithObjectsAndKeys:[NSNumber numberWithFloat:Options.CompressionQuality.GetValue()], AVVideoQualityKey, nil], AVVideoCompressionPropertiesKey,
-								 [NSNumber numberWithInt:Viewport->GetSize().X], AVVideoWidthKey,
-								 [NSNumber numberWithInt:Viewport->GetSize().Y], AVVideoHeightKey,
+								 [NSNumber numberWithInt:Options.Width], AVVideoWidthKey,
+								 [NSNumber numberWithInt:Options.Height], AVVideoHeightKey,
 								 nil];
 			}
 			else
 			{
 				VideoSettings = [NSDictionary dictionaryWithObjectsAndKeys:
 								 AVVideoCodecH264, AVVideoCodecKey,
-								 [NSNumber numberWithInt:Viewport->GetSize().X], AVVideoWidthKey,
-								 [NSNumber numberWithInt:Viewport->GetSize().Y], AVVideoHeightKey,
+								 [NSNumber numberWithInt:Options.Width], AVVideoWidthKey,
+								 [NSNumber numberWithInt:Options.Height], AVVideoHeightKey,
 								 nil];
 			}
 			AVFWriterInputRef = [[AVAssetWriterInput
@@ -477,7 +468,7 @@ public:
 		}
 	}
 	
-	void StopCapture()
+	void Finalize()
 	{
 		if (bCapturing)
 		{
@@ -493,16 +484,10 @@ public:
 			AVFPixelBufferAdaptorRef = nil;
 
 			bCapturing = false;
-			CaptureViewport = nullptr;
 			FrameNumber = 0;
 		}
 	}
 	
-	void Close()
-	{
-		StopCapture();
-	}
-	
 	void TaskThread()
 	{
 		SCOPED_AUTORELEASE_POOL;
@@ -512,12 +497,6 @@ public:
 			uint32 WaitTimeMs = 100;
 			TArray<FCapturedFrame> PendingFrames = GetFrameData(WaitTimeMs);
 
-			auto Viewport = CaptureViewport.Pin();
-			if (!Viewport.IsValid())
-			{
-				break;
-			}
-
 			// Capture the frames that we have
 			for (auto& CurrentFrame : PendingFrames)
 			{
@@ -530,7 +509,7 @@ public:
 				CVPixelBufferPoolCreatePixelBuffer (NULL, AVFPixelBufferAdaptorRef.pixelBufferPool, &PixelBuffer);
 				if(!PixelBuffer)
 				{
-					CVPixelBufferCreate(kCFAllocatorDefault, Viewport->GetSize().X, Viewport->GetSize().Y, kCVPixelFormatType_32BGRA, NULL, &PixelBuffer);
+					CVPixelBufferCreate(kCFAllocatorDefault, Options.Width, Options.Height, kCVPixelFormatType_32BGRA, NULL, &PixelBuffer);
 				}
 				check(PixelBuffer);
 				
@@ -602,11 +581,9 @@ FCapturedFrames::FCapturedFrames(const FString& InArchiveDirectory, int32 InMaxI
 {
 	FrameReady = FPlatformProcess::GetSynchEventFromPool();
 
-	// Ensure the file doesn't exist
+	// Ensure the archive directory doesn't exist
 	auto& PlatformFile = FPlatformFileManager::Get().GetPlatformFile();
-
 	PlatformFile.DeleteDirectoryRecursively(*ArchiveDirectory);
-	PlatformFile.CreateDirectory(*ArchiveDirectory);
 
 	TotalArchivedFrames = 0;
 	InMemoryFrames.Reserve(MaxInMemoryFrames);
@@ -660,6 +637,12 @@ FArchive &operator<<(FArchive& Ar, FCapturedFrame& Frame)
 
 void FCapturedFrames::ArchiveFrame(FCapturedFrame Frame)
 {
+	auto& PlatformFile = FPlatformFileManager::Get().GetPlatformFile();
+	if (!PlatformFile.DirectoryExists(*ArchiveDirectory))
+	{
+		PlatformFile.CreateDirectory(*ArchiveDirectory);
+	}
+
 	// Get (and increment) a unique index for this frame
 	uint32 ArchivedFrameIndex = ++TotalArchivedFrames;
 
diff --git a/Engine/Source/Runtime/Engine/Private/GameEngine.cpp b/Engine/Source/Runtime/Engine/Private/GameEngine.cpp
--- a/Engine/Source/Runtime/Engine/Private/GameEngine.cpp
+++ b/Engine/Source/Runtime/Engine/Private/GameEngine.cpp
@@ -353,7 +353,6 @@ void UGameEngine::SwitchGameWindowToUseGameViewport()
 		
 		SceneViewport->ResizeFrame((uint32)GSystemResolution.ResX, (uint32)GSystemResolution.ResY, GSystemResolution.WindowMode, 0, 0);
 
-
 		// Move the registration of the game viewport to that messages are correctly received.
 		if (!FPlatformProperties::SupportsWindowedMode())
 		{
@@ -704,6 +703,25 @@ bool UGameEngine::Exec( UWorld* InWorld, const TCHAR* Cmd, FOutputDevice& Ar )
 		return HandleApplyUserSettingsCommand( Cmd, Ar );
 	}
 #endif // !UE_BUILD_SHIPPING
+#if WITH_EDITOR
+	else if( FParse::Command(&Cmd,TEXT("STARTMOVIECAPTURE")) && GIsEditor )
+	{
+		IMovieSceneCaptureInterface* CaptureInterface = IMovieSceneCaptureModule::Get().GetFirstActiveMovieSceneCapture();
+		if (CaptureInterface)
+		{
+			CaptureInterface->StartCapturing();
+			return true;
+		}
+		else if (SceneViewport.IsValid())
+		{
+			if (IMovieSceneCaptureModule::Get().CreateMovieSceneCapture(SceneViewport))
+			{
+				return true;
+			}
+		}
+		return false;
+	}
+#endif
 	else if( InWorld && InWorld->Exec( InWorld, Cmd, Ar ) )
 	{
 		return true;
diff --git a/Engine/Source/Runtime/Engine/Private/GameInstance.cpp b/Engine/Source/Runtime/Engine/Private/GameInstance.cpp
--- a/Engine/Source/Runtime/Engine/Private/GameInstance.cpp
+++ b/Engine/Source/Runtime/Engine/Private/GameInstance.cpp
@@ -320,12 +320,7 @@ void UGameInstance::StartGameInstance()
 	const FString& DefaultMap = GameMapsSettings->GetGameDefaultMap();
 
 	FString PackageName;
-
-#if WITH_EDITOR
-	PackageName = InitialMapOverride;
-#endif
-
-	if (PackageName.IsEmpty() && (!FParse::Token(Tmp, PackageName, 0) || **PackageName == '-'))
+	if (!FParse::Token(Tmp, PackageName, 0) || **PackageName == '-')
 	{
 		PackageName = DefaultMap + GameMapsSettings->LocalMapOptions;
 	}
diff --git a/Engine/Source/Runtime/Engine/Private/HighResScreenshot.cpp b/Engine/Source/Runtime/Engine/Private/HighResScreenshot.cpp
--- a/Engine/Source/Runtime/Engine/Private/HighResScreenshot.cpp
+++ b/Engine/Source/Runtime/Engine/Private/HighResScreenshot.cpp
@@ -23,13 +23,19 @@ FHighResScreenshotConfig::FHighResScreenshotConfig()
 	SetHDRCapture(false);
 }
 
-void FHighResScreenshotConfig::Init()
+void FHighResScreenshotConfig::Init(uint32 NumAsyncWriters)
 {
 	IImageWrapperModule* ImageWrapperModule = FModuleManager::LoadModulePtr<IImageWrapperModule>(FName("ImageWrapper"));
 	if (ImageWrapperModule != nullptr)
 	{
-		ImageCompressorLDR = ImageWrapperModule->CreateImageWrapper(EImageFormat::PNG);
-		ImageCompressorHDR = ImageWrapperModule->CreateImageWrapper(EImageFormat::EXR);
+		ImageCompressorsLDR.Reserve(NumAsyncWriters);
+		ImageCompressorsHDR.Reserve(NumAsyncWriters);
+
+		for (uint32 Index = 0; Index != NumAsyncWriters; ++Index)
+		{
+			ImageCompressorsLDR.Emplace(ImageWrapperModule->CreateImageWrapper(EImageFormat::PNG));
+			ImageCompressorsHDR.Emplace(ImageWrapperModule->CreateImageWrapper(EImageFormat::EXR));
+		}
 	}
 
 #if WITH_EDITOR
@@ -160,7 +166,30 @@ bool FHighResScreenshotConfig::SaveImage(const FString& File, const TArray<TPixe
 	IFileManager* FileManager = &IFileManager::Get();
 	const size_t BitsPerPixel = (sizeof(TPixelType) / 4) * 8;
 
-	TSharedPtr<class IImageWrapper> ImageCompressor = bIsWritingHDRImage ? ImageCompressorHDR : ImageCompressorLDR;
+	const FImageWriter* ImageWriter = nullptr;
+	{
+		const TArray<FImageWriter>& ArrayToUse = bIsWritingHDRImage ? ImageCompressorsHDR : ImageCompressorsLDR;
+
+		// Find a free image writer to use. This can potentially be called on many threads at the same time
+		for (;;)
+		{
+			for (const FImageWriter& Writer : ArrayToUse)
+			{
+				if (!Writer.bInUse.AtomicSet(true))
+				{
+					ImageWriter = &Writer;
+					break;
+				}
+			}
+
+			if (ImageWriter)
+			{
+				break;
+			}
+
+			FPlatformProcess::Sleep(0.001f);
+		}
+	}
 
 	// here we require the input file name to have an extension
 	FString NewExtension = bIsWritingHDRImage ? TEXT(".exr") : TEXT(".png");
@@ -170,25 +199,32 @@ bool FHighResScreenshotConfig::SaveImage(const FString& File, const TArray<TPixe
 	{
 		*OutFilename = Filename;
 	}
+	
+	bool bSuccess = false;
 
-	if (ImageCompressor.IsValid() && ImageCompressor->SetRaw((void*)&Bitmap[0], sizeof(TPixelType)* x * y, x, y, Traits::SourceChannelLayout, BitsPerPixel))
+	if (ImageWriter && ImageWriter->ImageWrapper->SetRaw((void*)&Bitmap[0], sizeof(TPixelType)* x * y, x, y, Traits::SourceChannelLayout, BitsPerPixel))
 	{
 		FArchive* Ar = FileManager->CreateFileWriter(Filename.GetCharArray().GetData());
 		if (Ar != nullptr)
 		{
-			const TArray<uint8>& CompressedData = ImageCompressor->GetCompressed();
+			const TArray<uint8>& CompressedData = ImageWriter->ImageWrapper->GetCompressed();
 			int32 CompressedSize = CompressedData.Num();
 			Ar->Serialize((void*)CompressedData.GetData(), CompressedSize);
 			delete Ar;
-		}
-		else
-		{
-			return false;
+
+			bSuccess = true;
 		}
 	}
 
-	return true;
+	ImageWriter->bInUse = false;
+
+	return bSuccess;
 }
 
 template ENGINE_API bool FHighResScreenshotConfig::SaveImage<FColor>(const FString& File, const TArray<FColor>& Bitmap, const FIntPoint& BitmapSize, FString* OutFilename) const;
 template ENGINE_API bool FHighResScreenshotConfig::SaveImage<FFloat16Color>(const FString& File, const TArray<FFloat16Color>& Bitmap, const FIntPoint& BitmapSize, FString* OutFilename) const;
+
+FHighResScreenshotConfig::FImageWriter::FImageWriter(const TSharedPtr<class IImageWrapper>& InWrapper)
+	: ImageWrapper(InWrapper)
+{
+}
diff --git a/Engine/Source/Runtime/Engine/Private/Slate/SceneViewport.cpp b/Engine/Source/Runtime/Engine/Private/Slate/SceneViewport.cpp
--- a/Engine/Source/Runtime/Engine/Private/Slate/SceneViewport.cpp
+++ b/Engine/Source/Runtime/Engine/Private/Slate/SceneViewport.cpp
@@ -856,6 +856,14 @@ FReply FSceneViewport::OnFocusReceived(const FFocusEvent& InFocusEvent)
 		}
 	}
 
+	// Update key state mappings so that the the viewport modifier states are valid upon focus.
+	KeyStateMap.Add( EKeys::LeftAlt, FSlateApplication::Get().GetModifierKeys().IsLeftAltDown() );
+	KeyStateMap.Add( EKeys::RightAlt, FSlateApplication::Get().GetModifierKeys().IsRightAltDown() );
+	KeyStateMap.Add( EKeys::LeftControl, FSlateApplication::Get().GetModifierKeys().IsLeftControlDown());
+	KeyStateMap.Add( EKeys::RightControl, FSlateApplication::Get().GetModifierKeys().IsRightControlDown());
+	KeyStateMap.Add( EKeys::LeftShift, FSlateApplication::Get().GetModifierKeys().IsLeftShiftDown());
+	KeyStateMap.Add( EKeys::RightShift, FSlateApplication::Get().GetModifierKeys().IsRightShiftDown());
+
 	return CurrentReplyState;
 }
 
diff --git a/Engine/Source/Runtime/Engine/Private/UnrealEngine.cpp b/Engine/Source/Runtime/Engine/Private/UnrealEngine.cpp
--- a/Engine/Source/Runtime/Engine/Private/UnrealEngine.cpp
+++ b/Engine/Source/Runtime/Engine/Private/UnrealEngine.cpp
@@ -1051,6 +1051,11 @@ void UEngine::ShutdownAudioDeviceManager()
 
 void UEngine::PreExit()
 {
+	if (IMovieSceneCaptureModule* Module = FModuleManager::GetModulePtr<IMovieSceneCaptureModule>("MovieSceneCapture"))
+	{
+		Module->DestroyAllActiveCaptures();
+	}
+
 	ShutdownRenderingCVarsCaching();
 	FEngineAnalytics::Shutdown();
 	if (ScreenSaverInhibitor)
@@ -2391,10 +2396,6 @@ bool UEngine::Exec( UWorld* InWorld, const TCHAR* Cmd, FOutputDevice& Ar )
 	{
 		return HandleStatCommand(InWorld, GStatProcessingViewportClient, Cmd, Ar);
 	}
-	else if( FParse::Command(&Cmd,TEXT("STARTMOVIECAPTURE")) && GIsEditor )
-	{
-		return HandleStartMovieCaptureCommand( Cmd, Ar );
-	}
 	else if( FParse::Command(&Cmd,TEXT("STOPMOVIECAPTURE")) && GIsEditor )
 	{
 		return HandleStopMovieCaptureCommand( Cmd, Ar );
@@ -2746,27 +2747,11 @@ bool UEngine::Exec( UWorld* InWorld, const TCHAR* Cmd, FOutputDevice& Ar )
 	return true;
 }
 
-FMovieSceneCaptureHandle GMovieCaptureHandle;
-bool UEngine::HandleStartMovieCaptureCommand( const TCHAR* Cmd, FOutputDevice& Ar )
-{
-	if (!GMovieCaptureHandle.IsValid())
-	{
-		// IMovieSceneCaptureInterface* CaptureInterface = IMovieSceneCaptureModule::Get().CreateMovieSceneCapture(GameViewport);
-		// if (CaptureInterface)
-		// {
-		// 	GMovieCaptureHandle = CaptureInterface->GetHandle();
-		// 	return true;
-		// }
-	}
-	return false;
-}
-
 bool UEngine::HandleStopMovieCaptureCommand( const TCHAR* Cmd, FOutputDevice& Ar )
 {
-	if (GMovieCaptureHandle.IsValid())
+	if (IMovieSceneCaptureInterface* CaptureInterface = IMovieSceneCaptureModule::Get().GetFirstActiveMovieSceneCapture())
 	{
-		IMovieSceneCaptureModule::Get().DestroyMovieSceneCapture(GMovieCaptureHandle);
-		GMovieCaptureHandle = FMovieSceneCaptureHandle();
+		CaptureInterface->Close();
 		return true;
 	}
 	return false;
diff --git a/Engine/Source/Runtime/Engine/Public/AVIWriter.h b/Engine/Source/Runtime/Engine/Public/AVIWriter.h
--- a/Engine/Source/Runtime/Engine/Public/AVIWriter.h
+++ b/Engine/Source/Runtime/Engine/Public/AVIWriter.h
@@ -6,9 +6,6 @@
 
 #pragma once
 
-#include "UnrealClient.h"
-#include "SceneViewport.h"
-
 DECLARE_LOG_CATEGORY_EXTERN(LogMovieCapture, Warning, All);
 
 DECLARE_DELEGATE_RetVal_TwoParams(FString, FResolveFileFormat, const FString&, const FString&);
@@ -20,6 +17,8 @@ struct FAVIWriterOptions
 		: OutputFilename(FPaths::VideoCaptureDir() / TEXT("Capture.avi"))
 		, CaptureFPS(30)
 		, bSynchronizeFrames(false)
+		, Width(0)
+		, Height(0)
 	{}
 
 	/** Output filename */
@@ -36,6 +35,10 @@ struct FAVIWriterOptions
 
 	/** When true, the game thread will block until captured frames have been processed by the avi writer */
 	bool bSynchronizeFrames;
+
+	uint32 Width;
+
+	uint32 Height;
 };
 
 /** Data structure representing a captured frame */
@@ -117,14 +120,13 @@ private:
 	TOptional<TFuture<void>> UnarchiveTask;
 };
 
-/** Class responsible for capturing frames from a viewport, and writing them out to an AVI file */
+/** Class responsible for writing frames out to an AVI file */
 class FAVIWriter
 {
 protected:
 	/** Protected constructor to avoid abuse. */
 	FAVIWriter(const FAVIWriterOptions& InOptions) 
 		: bCapturing(false)
-		, CaptureViewport(nullptr)
 		, FrameNumber(0)
 		, Options(InOptions)
 	{
@@ -133,9 +135,6 @@ protected:
 	/** Whether we are capturing or not */
 	FThreadSafeBool bCapturing;
 
-	/** The viewport we are capturing from */
-	TWeakPtr<FSceneViewport> CaptureViewport;
-
 	/** The current frame number */
 	int32 FrameNumber;
 	
@@ -167,14 +166,12 @@ public:
 
 	uint32 GetWidth() const
 	{
-		auto Pin = CaptureViewport.Pin();
-		return Pin.IsValid() ? Pin->GetSize().X : 0;
+		return Options.Width;
 	}
 
 	uint32 GetHeight() const
 	{
-		auto Pin = CaptureViewport.Pin();
-		return Pin.IsValid() ? Pin->GetSize().Y : 0;
+		return Options.Height;
 	}
 
 	int32 GetFrameNumber() const
@@ -189,11 +186,8 @@ public:
 
 	ENGINE_API void Update(double FrameTimeSeconds, TArray<FColor> FrameData);
 	
-	virtual void StartCapture(TWeakPtr<FSceneViewport> Viewport) = 0;
-	virtual void StopCapture() = 0;
-	virtual void Close() = 0;
+	virtual void Initialize() = 0;
+	virtual void Finalize() = 0;
+	
 	virtual void DropFrames(int32 NumFramesToDrop) = 0;
-
-protected:
-	void PrepareForScreenshot();
 };
\ No newline at end of file
diff --git a/Engine/Source/Runtime/Engine/Public/HighResScreenshot.h b/Engine/Source/Runtime/Engine/Public/HighResScreenshot.h
--- a/Engine/Source/Runtime/Engine/Public/HighResScreenshot.h
+++ b/Engine/Source/Runtime/Engine/Public/HighResScreenshot.h
@@ -16,9 +16,6 @@ struct ENGINE_API FHighResScreenshotConfig
 	bool bDisplayCaptureRegion;
 	bool bCaptureHDR;
 
-	TSharedPtr<class IImageWrapper> ImageCompressorLDR;
-	TSharedPtr<class IImageWrapper> ImageCompressorHDR;
-
 	// Materials used in the editor to help with the capture of highres screenshots
 	UMaterial* HighResScreenshotMaterial;
 	UMaterial* HighResScreenshotMaskMaterial;
@@ -27,7 +24,7 @@ struct ENGINE_API FHighResScreenshotConfig
 	FHighResScreenshotConfig();
 
 	/** Initialize the image wrapper modules (required for SaveImage) **/
-	void Init();
+	void Init(uint32 NumAsyncWriters = 6);
 
 	/** Point the screenshot UI at a different viewport **/
 	void ChangeViewport(TWeakPtr<FSceneViewport> InViewport);
@@ -46,6 +43,17 @@ struct ENGINE_API FHighResScreenshotConfig
 	/** Save to image file **/
 	template<typename TPixelType>
 	ENGINE_API bool SaveImage(const FString& File, const TArray<TPixelType>& Bitmap, const FIntPoint& BitmapSize, FString* OutFilename = nullptr) const;
+
+private:
+	struct FImageWriter
+	{
+		FImageWriter(const TSharedPtr<class IImageWrapper>& InWrapper);
+
+		TSharedPtr<class IImageWrapper> ImageWrapper;
+		mutable FThreadSafeBool bInUse;
+	};
+	TArray<FImageWriter> ImageCompressorsLDR;
+	TArray<FImageWriter> ImageCompressorsHDR;
 };
 
 ENGINE_API FHighResScreenshotConfig& GetHighResScreenshotConfig();
\ No newline at end of file
diff --git a/Engine/Source/Runtime/LevelSequence/Private/LevelSequence.cpp b/Engine/Source/Runtime/LevelSequence/Private/LevelSequence.cpp
--- a/Engine/Source/Runtime/LevelSequence/Private/LevelSequence.cpp
+++ b/Engine/Source/Runtime/LevelSequence/Private/LevelSequence.cpp
@@ -65,8 +65,6 @@ void ULevelSequence::BindPossessableObject(const FGuid& ObjectId, UObject& Posse
 	if (Context)
 	{
 		ObjectReferences.CreateBinding(ObjectId, &PossessedObject, Context);
-
-		MovieSceneHelpers::SetRuntimeObjectMobility(&PossessedObject);
 	}
 }
 
diff --git a/Engine/Source/Runtime/LevelSequence/Private/LevelSequencePlayer.cpp b/Engine/Source/Runtime/LevelSequence/Private/LevelSequencePlayer.cpp
--- a/Engine/Source/Runtime/LevelSequence/Private/LevelSequencePlayer.cpp
+++ b/Engine/Source/Runtime/LevelSequence/Private/LevelSequencePlayer.cpp
@@ -68,6 +68,8 @@ ULevelSequencePlayer::ULevelSequencePlayer(const FObjectInitializer& ObjectIniti
 	, LevelSequence(nullptr)
 	, bIsPlaying(false)
 	, TimeCursorPosition(0.0f)
+	, StartTime(0.f)
+	, EndTime(0.f)
 	, CurrentNumLoops(0)
 	, bHasCleanedUpSequence(false)
 {
@@ -164,13 +166,7 @@ void ULevelSequencePlayer::SetPlaybackPosition(float NewPlaybackPosition)
 
 float ULevelSequencePlayer::GetLength() const
 {
-	if (!LevelSequence)
-	{
-		return 0;
-	}
-
-	UMovieScene* MovieScene = LevelSequence->GetMovieScene();
-	return MovieScene ? MovieScene->GetPlaybackRange().Size<float>() : 0;
+	return EndTime - StartTime;
 }
 
 float ULevelSequencePlayer::GetPlayRate() const
@@ -185,14 +181,10 @@ void ULevelSequencePlayer::SetPlayRate(float PlayRate)
 
 void ULevelSequencePlayer::SetPlaybackRange( const float NewStartTime, const float NewEndTime )
 {
-	if( LevelSequence != nullptr )
-	{
-		UMovieScene* MovieScene = LevelSequence->GetMovieScene();
-		if( MovieScene != nullptr )
-		{
-			MovieScene->SetPlaybackRange( NewStartTime, NewEndTime );
-		}
-	}
+	StartTime = NewStartTime;
+	EndTime = FMath::Max(NewEndTime, StartTime);
+
+	TimeCursorPosition = FMath::Clamp(TimeCursorPosition, 0.f, GetLength());
 }
 
 void ULevelSequencePlayer::OnCursorPositionChanged()
@@ -230,6 +222,12 @@ void ULevelSequencePlayer::Initialize(ULevelSequence* InLevelSequence, UWorld* I
 	World = InWorld;
 	PlaybackSettings = Settings;
 
+	if (UMovieScene* MovieScene = LevelSequence->GetMovieScene())
+	{
+		TRange<float> PlaybackRange = MovieScene->GetPlaybackRange();
+		SetPlaybackRange(PlaybackRange.GetLowerBoundValue(), PlaybackRange.GetUpperBoundValue());
+	}
+
 	// Ensure everything is set up, ready for playback
 	Stop();
 }
@@ -341,8 +339,7 @@ void ULevelSequencePlayer::UpdateMovieSceneInstance(float CurrentPosition, float
 {
 	if(RootMovieSceneInstance.IsValid())
 	{
-		const float SequenceStartOffset = LevelSequence->GetMovieScene()->GetPlaybackRange().GetLowerBoundValue();
-		RootMovieSceneInstance->Update(CurrentPosition + SequenceStartOffset, PreviousPosition + SequenceStartOffset, *this);
+		RootMovieSceneInstance->Update(CurrentPosition + StartTime, PreviousPosition + StartTime, *this);
 #if WITH_EDITOR
 		OnLevelSequencePlayerUpdate.Broadcast(*this, CurrentPosition, PreviousPosition);
 #endif
diff --git a/Engine/Source/Runtime/LevelSequence/Private/LevelSequenceSpawnRegister.cpp b/Engine/Source/Runtime/LevelSequence/Private/LevelSequenceSpawnRegister.cpp
--- a/Engine/Source/Runtime/LevelSequence/Private/LevelSequenceSpawnRegister.cpp
+++ b/Engine/Source/Runtime/LevelSequence/Private/LevelSequenceSpawnRegister.cpp
@@ -87,8 +87,6 @@ UObject* FLevelSequenceSpawnRegister::SpawnObject(const FGuid& BindingId, FMovie
 		return nullptr;
 	}
 
-	MovieSceneHelpers::SetRuntimeObjectMobility(SpawnedObject);
-
 	Register.Add(Key, FSpawnedObject(*SpawnedObject, Spawnable->GetSpawnOwnership()));
 
 	SequenceInstance.OnObjectSpawned(BindingId, *SpawnedObject, Player);
diff --git a/Engine/Source/Runtime/LevelSequence/Public/LevelSequencePlayer.h b/Engine/Source/Runtime/LevelSequence/Public/LevelSequencePlayer.h
--- a/Engine/Source/Runtime/LevelSequence/Public/LevelSequencePlayer.h
+++ b/Engine/Source/Runtime/LevelSequence/Public/LevelSequencePlayer.h
@@ -106,7 +106,7 @@ public:
 	UFUNCTION(BlueprintCallable, Category="Game|Cinematic")
 	bool IsPlaying() const;
 
-	/** Get the length of the sequence */
+	/** Get the playback length of the sequence */
 	UFUNCTION(BlueprintCallable, Category="Game|Cinematic")
 	float GetLength() const;
 
@@ -148,9 +148,6 @@ public:
 
 	void Update(const float DeltaSeconds);
 
-	/** Instruct this player to start playing next frame update */
-	void AutoPlayNextFrame();
-
 private:
 
 	/** Update the movie scene instance from the specified previous position, to the specified time position */
@@ -173,6 +170,12 @@ private:
 	UPROPERTY()
 	float TimeCursorPosition;
 
+	/** Time time at which to start playing the sequence (defaults to the lower bound of the sequence's play range) */
+	float StartTime;
+
+	/** Time time at which to end playing the sequence (defaults to the upper bound of the sequence's play range) */
+	float EndTime;
+
 	/** Specific playback settings for the animation. */
 	UPROPERTY()
 	FLevelSequencePlaybackSettings PlaybackSettings;
diff --git a/Engine/Source/Runtime/Media/Public/IMediaPlayer.h b/Engine/Source/Runtime/Media/Public/IMediaPlayer.h
--- a/Engine/Source/Runtime/Media/Public/IMediaPlayer.h
+++ b/Engine/Source/Runtime/Media/Public/IMediaPlayer.h
@@ -12,6 +12,37 @@ class IMediaCaptionTrack;
 class IMediaVideoTrack;
 
 
+/**
+ * Enumerates media player related events.
+ */
+enum class EMediaEvent
+{
+	/** Unknown event. */
+	Unknown,
+
+	/** The current media source has been closed. */
+	MediaClosed,
+
+	/** A new media source has been opened. */
+	MediaOpened,
+
+	/** A media source failed to open. */
+	MediaOpenFailed,
+
+	/** The end of the media (or beginning if playing in reverse) has been reached. */
+	PlaybackEndReached,
+
+	/** Playback has been resumed. */
+	PlaybackResumed,
+
+	/** Playback has been suspended. */
+	PlaybackSuspended,
+
+	/** Media tracks have changed. */
+	TracksChanged
+};
+
+
 /**
  * Enumerates directions for seeking in media.
  */
@@ -201,21 +232,9 @@ public:
 
 public:
 
-	/** Gets an event delegate that is invoked when media has been closed. */
-	DECLARE_EVENT(IMediaPlayer, FOnMediaClosed)
-	virtual FOnMediaClosed& OnClosed() = 0;
-
-	/** Gets an event delegate that is invoked when media has been opened. */
-	DECLARE_EVENT_OneParam(IMediaPlayer, FOnMediaOpened, FString /*OpenedUrl*/)
-	virtual FOnMediaOpened& OnOpened() = 0;
-
-	/** Gets an event delegate that is invoked when media failed to open. */
-	DECLARE_EVENT_OneParam(IMediaPlayer, FOnMediaOpenFailed, FString /*FailedUrl*/)
-	virtual FOnMediaOpenFailed& OnOpenFailed() = 0;
-
-	/** Gets an event delegate that is invoked when the media tracks have changed. */
-	DECLARE_EVENT(IMediaPlayer, FOnTracksChanged)
-	virtual FOnTracksChanged& OnTracksChanged() = 0;
+	/** Gets an event delegate that is invoked when some interesting event occurred. */
+	DECLARE_EVENT_OneParam(IMediaPlayer, FOnMediaEvent, EMediaEvent /*Event*/)
+	virtual FOnMediaEvent& OnMediaEvent() = 0;
 
 public:
 
diff --git a/Engine/Source/Runtime/MediaAssets/MediaAssets.Build.cs b/Engine/Source/Runtime/MediaAssets/MediaAssets.Build.cs
--- a/Engine/Source/Runtime/MediaAssets/MediaAssets.Build.cs
+++ b/Engine/Source/Runtime/MediaAssets/MediaAssets.Build.cs
@@ -6,17 +6,12 @@ namespace UnrealBuildTool.Rules
 	{
 		public MediaAssets(TargetInfo Target)
 		{
-			DynamicallyLoadedModuleNames.AddRange(
-				new string[] {
-					"Media",
-				}
-			);
-
             PublicDependencyModuleNames.AddRange(
                 new string[] {
 					"Core",
 					"CoreUObject",
-					"Engine",                
+					"Engine",
+                    "Media",
                 }
             );
 
@@ -27,12 +22,6 @@ namespace UnrealBuildTool.Rules
 				}
 			);
 
-			PrivateIncludePathModuleNames.AddRange(
-				new string[] {
-					"Media",
-				}
-			);
-
 			PrivateIncludePaths.AddRange(
 				new string[] {
 					"Runtime/MediaAssets/Private",
diff --git a/Engine/Source/Runtime/MediaAssets/Private/Assets/MediaPlayer.cpp b/Engine/Source/Runtime/MediaAssets/Private/Assets/MediaPlayer.cpp
--- a/Engine/Source/Runtime/MediaAssets/Private/Assets/MediaPlayer.cpp
+++ b/Engine/Source/Runtime/MediaAssets/Private/Assets/MediaPlayer.cpp
@@ -8,7 +8,7 @@
 /* UMediaPlayer structors
  *****************************************************************************/
 
-UMediaPlayer::UMediaPlayer( const FObjectInitializer& ObjectInitializer )
+UMediaPlayer::UMediaPlayer(const FObjectInitializer& ObjectInitializer)
 	: Super(ObjectInitializer)
 	, Looping(true)
 	, Player(nullptr)
@@ -96,13 +96,13 @@ bool UMediaPlayer::IsPlaying() const
 }
 
 
-bool UMediaPlayer::IsStopped() const
+bool UMediaPlayer::IsReady() const
 {
-	return !Player.IsValid() || !Player->IsReady();
+	return Player.IsValid() && Player->IsReady();
 }
 
 
-bool UMediaPlayer::OpenUrl( const FString& NewUrl )
+bool UMediaPlayer::OpenUrl(const FString& NewUrl)
 {
 	URL = NewUrl;
 	InitializePlayer();
@@ -129,25 +129,25 @@ bool UMediaPlayer::Rewind()
 }
 
 
-bool UMediaPlayer::SetLooping( bool InLooping )
+bool UMediaPlayer::SetLooping(bool InLooping)
 {
 	return Player.IsValid() && Player->SetLooping(InLooping);
 }
 
 
-bool UMediaPlayer::SetRate( float Rate )
+bool UMediaPlayer::SetRate(float Rate)
 {
 	return Player.IsValid() && Player->SetRate(Rate);
 }
 
 
-bool UMediaPlayer::Seek( const FTimespan& InTime )
+bool UMediaPlayer::Seek(const FTimespan& InTime)
 {
 	return Player.IsValid() && Player->Seek(InTime);
 }
 
 
-bool UMediaPlayer::SupportsRate( float Rate, bool Unthinned ) const
+bool UMediaPlayer::SupportsRate(float Rate, bool Unthinned) const
 {
 	return Player.IsValid() && Player->GetMediaInfo().SupportsRate(Rate, Unthinned);
 }
@@ -204,7 +204,7 @@ void UMediaPlayer::PostLoad()
 
 #if WITH_EDITOR
 
-void UMediaPlayer::PostEditChangeProperty( FPropertyChangedEvent& PropertyChangedEvent )
+void UMediaPlayer::PostEditChangeProperty(FPropertyChangedEvent& PropertyChangedEvent)
 {
 	Super::PostEditChangeProperty(PropertyChangedEvent);
 
@@ -224,74 +224,65 @@ void UMediaPlayer::InitializePlayer()
 		return;
 	}
 
-	if (URL != CurrentUrl)
+	if (URL == CurrentUrl)
 	{
-		// close previous player
-		CurrentUrl = FString();
-
-		if (Player.IsValid())
-		{
-			Player->Close();
-			Player->OnClosed().RemoveAll(this);
-			Player->OnOpened().RemoveAll(this);
-			Player->OnOpenFailed().RemoveAll(this);
-			Player->OnTracksChanged().RemoveAll(this);
-			Player.Reset();
-		}
-
-		if (URL.IsEmpty())
-		{
-			return;
-		}
-
-		// create new player
-		IMediaModule* MediaModule = FModuleManager::LoadModulePtr<IMediaModule>("Media");
-
-		if (MediaModule == nullptr)
-		{
-			return;
-		}
-
-		Player = MediaModule->CreatePlayer(URL);
-
-		if (!Player.IsValid())
-		{
-			return;
-		}
-
-		Player->OnClosed().AddUObject(this, &UMediaPlayer::HandlePlayerMediaClosed);
-		Player->OnOpened().AddUObject(this, &UMediaPlayer::HandlePlayerMediaOpened);
-		Player->OnOpenFailed().AddUObject(this, &UMediaPlayer::HandlePlayerMediaOpenFailed);
-		Player->OnTracksChanged().AddUObject(this, &UMediaPlayer::HandlePlayerTracksChanged);
-
-		// open the new media file
-		bool OpenedSuccessfully = false;
-
-		if (URL.Contains(TEXT("://")))
-		{
-			OpenedSuccessfully = Player->Open(URL);
-		}
-		else
-		{
-			const FString FullUrl = FPaths::ConvertRelativePathToFull(
-				FPaths::IsRelative(URL)
-					? FPaths::GameContentDir() / URL
-					: URL
-			);
-
-			OpenedSuccessfully = Player->Open(FullUrl);
-		}
-
-		// finish initialization
-		if (OpenedSuccessfully)
-		{
-			CurrentUrl = URL;
-		}
+		return;
 	}
 
+	// close previous player
+	CurrentUrl = FString();
+
 	if (Player.IsValid())
 	{
-		Player->SetLooping(Looping);
+		Player->Close();
+		Player->OnMediaEvent().RemoveAll(this);
+		Player.Reset();
+	}
+
+	if (URL.IsEmpty())
+	{
+		return;
+	}
+
+	// create new player
+	IMediaModule* MediaModule = FModuleManager::LoadModulePtr<IMediaModule>("Media");
+
+	if (MediaModule == nullptr)
+	{
+		return;
+	}
+
+	Player = MediaModule->CreatePlayer(URL);
+
+	if (!Player.IsValid())
+	{
+		return;
+	}
+
+	Player->OnMediaEvent().AddUObject(this, &UMediaPlayer::HandlePlayerMediaEvent);
+
+	// open the new media file
+	bool OpenedSuccessfully = false;
+
+	if (URL.Contains(TEXT("://")))
+	{
+		OpenedSuccessfully = Player->Open(URL);
+	}
+	else
+	{
+		const FString FullUrl = FPaths::ConvertRelativePathToFull(
+			FPaths::IsRelative(URL)
+				? FPaths::GameContentDir() / URL
+				: URL
+		);
+
+		OpenedSuccessfully = Player->Open(FullUrl);
+	}
+
+	// finish initialization
+	if (OpenedSuccessfully)
+	{
+		CurrentUrl = URL;
 	}
 }
 
@@ -299,27 +290,38 @@ void UMediaPlayer::InitializePlayer()
 /* UMediaPlayer callbacks
  *****************************************************************************/
 
-void UMediaPlayer::HandlePlayerMediaClosed()
+void UMediaPlayer::HandlePlayerMediaEvent(EMediaEvent Event)
 {
-	MediaChangedEvent.Broadcast();
-	OnMediaClosed.Broadcast();
-}
+	switch(Event)
+	{
+	case EMediaEvent::MediaClosed:
+		MediaChangedEvent.Broadcast();
+		OnMediaClosed.Broadcast();
+		break;
 
+	case EMediaEvent::MediaOpened:
+		Player->SetLooping(Looping);
+		MediaChangedEvent.Broadcast();
+		OnMediaOpened.Broadcast(URL);
+		break;
 
-void UMediaPlayer::HandlePlayerMediaOpened(FString OpenedUrl)
-{
-	MediaChangedEvent.Broadcast();
-	OnMediaOpened.Broadcast(OpenedUrl);
-}
+	case EMediaEvent::MediaOpenFailed:
+		OnMediaOpenFailed.Broadcast(URL);
+		break;
 
+	case EMediaEvent::PlaybackEndReached:
+		OnEndReached.Broadcast();
+		break;
 
-void UMediaPlayer::HandlePlayerMediaOpenFailed(FString FailedUrl)
-{
-	OnMediaOpenFailed.Broadcast(FailedUrl);
-}
+	case EMediaEvent::PlaybackResumed:
+		OnPlaybackResumed.Broadcast();
+		break;
 
+	case EMediaEvent::PlaybackSuspended:
+		OnPlaybackSuspended.Broadcast();
+		break;
 
-void UMediaPlayer::HandlePlayerTracksChanged()
-{
-	TracksChangedEvent.Broadcast();
+	case EMediaEvent::TracksChanged:
+		TracksChangedEvent.Broadcast();
+	}
 }
diff --git a/Engine/Source/Runtime/MediaAssets/Private/Assets/MediaSoundWave.cpp b/Engine/Source/Runtime/MediaAssets/Private/Assets/MediaSoundWave.cpp
--- a/Engine/Source/Runtime/MediaAssets/Private/Assets/MediaSoundWave.cpp
+++ b/Engine/Source/Runtime/MediaAssets/Private/Assets/MediaSoundWave.cpp
@@ -9,7 +9,7 @@
 /* UMediaSoundWave structors
  *****************************************************************************/
 
-UMediaSoundWave::UMediaSoundWave( const FObjectInitializer& ObjectInitializer )
+UMediaSoundWave::UMediaSoundWave(const FObjectInitializer& ObjectInitializer)
 	: Super(ObjectInitializer)
 	, AudioTrackIndex(INDEX_NONE)
 	, AudioQueue(MakeShareable(new FMediaSampleQueue))
@@ -34,7 +34,7 @@ TSharedPtr<IMediaPlayer> UMediaSoundWave::GetPlayer() const
 }
 
 
-void UMediaSoundWave::SetMediaPlayer( UMediaPlayer* InMediaPlayer )
+void UMediaSoundWave::SetMediaPlayer( UMediaPlayer* InMediaPlayer)
 {
 	MediaPlayer = InMediaPlayer;
 
@@ -45,7 +45,7 @@ void UMediaSoundWave::SetMediaPlayer( UMediaPlayer* InMediaPlayer )
 /* USoundWave overrides
  *****************************************************************************/
 
-int32 UMediaSoundWave::GeneratePCMData( uint8* PCMData, const int32 SamplesNeeded )
+int32 UMediaSoundWave::GeneratePCMData( uint8* PCMData, const int32 SamplesNeeded)
 {
 	// drain media sample queue
 	TSharedPtr<TArray<uint8>, ESPMode::ThreadSafe> Sample;
@@ -74,25 +74,25 @@ int32 UMediaSoundWave::GeneratePCMData( uint8* PCMData, const int32 SamplesNeede
 }
 
 
-FByteBulkData* UMediaSoundWave::GetCompressedData( FName Format )
+FByteBulkData* UMediaSoundWave::GetCompressedData( FName Format)
 {
 	return nullptr;
 }
 
 
-int32 UMediaSoundWave::GetResourceSizeForFormat( FName Format )
+int32 UMediaSoundWave::GetResourceSizeForFormat( FName Format)
 {
 	return 0;
 }
 
 
-void UMediaSoundWave::InitAudioResource( FByteBulkData& CompressedData )
+void UMediaSoundWave::InitAudioResource( FByteBulkData& CompressedData)
 {
 	check(false); // should never be pushing compressed data to this class
 }
 
 
-bool UMediaSoundWave::InitAudioResource( FName Format )
+bool UMediaSoundWave::InitAudioResource( FName Format)
 {
 	return true;
 }
@@ -101,7 +101,7 @@ bool UMediaSoundWave::InitAudioResource( FName Format )
 /* UObject overrides
  *****************************************************************************/
 
-void UMediaSoundWave::GetAssetRegistryTags( TArray<FAssetRegistryTag>& OutTags ) const
+void UMediaSoundWave::GetAssetRegistryTags( TArray<FAssetRegistryTag>& OutTags) const
 {
 }
 
@@ -112,7 +112,7 @@ SIZE_T UMediaSoundWave::GetResourceSize(EResourceSizeMode::Type Mode)
 }
 
 
-void UMediaSoundWave::Serialize( FArchive& Ar )
+void UMediaSoundWave::Serialize( FArchive& Ar)
 {
 	// do not call the USoundWave version of serialize
 	USoundBase::Serialize(Ar);
@@ -165,6 +165,7 @@ void UMediaSoundWave::InitializeTrack()
 		{
 			MediaPlayer->OnTracksChanged().AddUObject(this, &UMediaSoundWave::HandleMediaPlayerTracksChanged);
 		}
+
 		bSetupDelegates = true;
 	}
 
@@ -172,9 +173,11 @@ void UMediaSoundWave::InitializeTrack()
 	if (AudioTrack.IsValid())
 	{
 		AudioTrack->GetStream().RemoveSink(AudioQueue);
-	}
+		AudioTrack.Reset();
 
-	AudioTrack.Reset();
+		AudioQueue->Empty();
+		QueuedAudio.Empty();
+	}
 
 	// initialize from new track
 	if (MediaPlayer != nullptr)
diff --git a/Engine/Source/Runtime/MediaAssets/Private/Assets/MediaTextureResource.cpp b/Engine/Source/Runtime/MediaAssets/Private/Assets/MediaTextureResource.cpp
--- a/Engine/Source/Runtime/MediaAssets/Private/Assets/MediaTextureResource.cpp
+++ b/Engine/Source/Runtime/MediaAssets/Private/Assets/MediaTextureResource.cpp
@@ -7,7 +7,7 @@
 /* FMediaTextureResource structors
  *****************************************************************************/
 
-FMediaTextureResource::FMediaTextureResource( const class UMediaTexture* InOwner, const FMediaSampleBufferRef& InVideoBuffer )
+FMediaTextureResource::FMediaTextureResource(const UMediaTexture* InOwner, const FMediaSampleBufferRef& InVideoBuffer)
 	: Cleared(false)
 	, LastFrameTime(FTimespan::MinValue())
 	, Owner(InOwner)
diff --git a/Engine/Source/Runtime/MediaAssets/Private/Assets/MediaTextureResource.h b/Engine/Source/Runtime/MediaAssets/Private/Assets/MediaTextureResource.h
--- a/Engine/Source/Runtime/MediaAssets/Private/Assets/MediaTextureResource.h
+++ b/Engine/Source/Runtime/MediaAssets/Private/Assets/MediaTextureResource.h
@@ -24,7 +24,7 @@ public:
 	 * @param InOwner The Movie texture object to create a resource for (must not be nullptr).
 	 * @param InVideoBuffer The video sample buffer to use.
 	 */
-	FMediaTextureResource(const class UMediaTexture* InOwner, const TSharedRef<FMediaSampleBuffer, ESPMode::ThreadSafe>& InVideoBuffer);
+	FMediaTextureResource(const UMediaTexture* InOwner, const TSharedRef<FMediaSampleBuffer, ESPMode::ThreadSafe>& InVideoBuffer);
 
 public:
 
@@ -43,7 +43,7 @@ public:
 
 	// FDeferredUpdateResource overrides
 
-	virtual void UpdateDeferredResource( FRHICommandListImmediate& RHICmdList, bool bClearRenderTarget=true ) override;
+	virtual void UpdateDeferredResource(FRHICommandListImmediate& RHICmdList, bool bClearRenderTarget = true) override;
 
 private:
 
diff --git a/Engine/Source/Runtime/MediaAssets/Public/MediaPlayer.h b/Engine/Source/Runtime/MediaAssets/Public/MediaPlayer.h
--- a/Engine/Source/Runtime/MediaAssets/Public/MediaPlayer.h
+++ b/Engine/Source/Runtime/MediaAssets/Public/MediaPlayer.h
@@ -6,10 +6,11 @@
 
 
 class IMediaPlayer;
+enum class EMediaEvent;
 
 
-/** Multicast delegate that is invoked when a media player's media has been closed. */
-DECLARE_DYNAMIC_MULTICAST_DELEGATE(FOnMediaPlayerMediaClosed);
+/** Multicast delegate that is invoked when a media event occurred in the player. */
+DECLARE_DYNAMIC_MULTICAST_DELEGATE(FOnMediaPlayerMediaEvent);
 
 /** Multicast delegate that is invoked when a media player's media has been opened. */
 DECLARE_DYNAMIC_MULTICAST_DELEGATE_OneParam(FOnMediaPlayerMediaOpened, FString, OpenedUrl);
@@ -109,7 +110,7 @@ public:
 	 * Checks whether playback is currently paused.
 	 *
 	 * @return true if playback is paused, false otherwise.
-	 * @see CanPause, IsPlaying, IsStopped, Pause
+	 * @see CanPause, IsPlaying, IsReady, Pause
 	 */
 	UFUNCTION(BlueprintCallable, Category="Media|MediaPlayer")
 	bool IsPaused() const;
@@ -118,19 +119,19 @@ public:
 	 * Checks whether playback has started.
 	 *
 	 * @return true if playback has started, false otherwise.
-	 * @see CanPlay, IsPaused, IsStopped, Play
+	 * @see CanPlay, IsPaused, IsReady, Play
 	 */
 	UFUNCTION(BlueprintCallable, Category="Media|MediaPlayer")
 	bool IsPlaying() const;
 
 	/**
-	 * Checks whether playback has stopped.
+	 * Checks whether media is ready for playback.
 	 *
-	 * @return true if playback has stopped, false otherwise.
+	 * @return true if media is ready, false otherwise.
 	 * @see IsPaused, IsPlaying, Stop
 	 */
 	UFUNCTION(BlueprintCallable, Category="Media|MediaPlayer")
-	bool IsStopped() const;
+	bool IsReady() const;
 
 	/**
 	 * Opens the specified media URL.
@@ -140,7 +141,7 @@ public:
 	 * @see GetUrl, Close
 	 */
 	UFUNCTION(BlueprintCallable, Category="Media|MediaPlayer")
-	bool OpenUrl( const FString& NewUrl );
+	bool OpenUrl(const FString& NewUrl);
 
 	/**
 	 * Pauses media playback.
@@ -233,6 +234,32 @@ public:
 	UFUNCTION(BlueprintCallable, Category="Media|MediaPlayer")
 	bool SupportsSeeking() const;
 
+public:
+
+	/** Holds a delegate that is invoked when playback has reached the end of the media. */
+	UPROPERTY(BlueprintAssignable, Category="Media|MediaPlayer")
+	FOnMediaPlayerMediaEvent OnEndReached;
+
+	/** Holds a delegate that is invoked when a media source has been closed. */
+	UPROPERTY(BlueprintAssignable, Category="Media|MediaPlayer")
+	FOnMediaPlayerMediaEvent OnMediaClosed;
+
+	/** Holds a delegate that is invoked when a media source has been opened. */
+	UPROPERTY(BlueprintAssignable, Category="Media|MediaPlayer")
+	FOnMediaPlayerMediaOpened OnMediaOpened;
+
+	/** Holds a delegate that is invoked when a media source has failed to open. */
+	UPROPERTY(BlueprintAssignable, Category="Media|MediaPlayer")
+	FOnMediaPlayerMediaOpenFailed OnMediaOpenFailed;
+
+	/** Holds a delegate that is invoked when media playback has been resumed. */
+	UPROPERTY(BlueprintAssignable, Category="Media|MediaPlayer")
+	FOnMediaPlayerMediaEvent OnPlaybackResumed;
+
+	/** Holds a delegate that is invoked when media playback has been suspended. */
+	UPROPERTY(BlueprintAssignable, Category="Media|MediaPlayer")
+	FOnMediaPlayerMediaEvent OnPlaybackSuspended;
+
 public:
 
 	/** Gets an event delegate that is invoked when media has been opened or closed. */
@@ -249,18 +276,6 @@ public:
 		return TracksChangedEvent;
 	}
 
-	/** Holds a delegate that is invoked when a media source has been closed. */
-	UPROPERTY(BlueprintAssignable, Category="Media|MediaPlayer")
-	FOnMediaPlayerMediaClosed OnMediaClosed;
-
-	/** Holds a delegate that is invoked when a media source has been opened. */
-	UPROPERTY(BlueprintAssignable, Category="Media|MediaPlayer")
-	FOnMediaPlayerMediaOpened OnMediaOpened;
-
-	/** Holds a delegate that is invoked when a media source has failed to open. */
-	UPROPERTY(BlueprintAssignable, Category="Media|MediaPlayer")
-	FOnMediaPlayerMediaOpenFailed OnMediaOpenFailed;
-
 public:
 
 	/**
@@ -302,17 +317,8 @@ protected:
 
 private:
 
-	/** Callback for when the player has closed a media source. */
-	void HandlePlayerMediaClosed();
-
-	/** Callback for when the player has opened a new media source. */
-	void HandlePlayerMediaOpened(FString OpenedUrl);
-
-	/** Callback for when the player has failed to open a new media source. */
-	void HandlePlayerMediaOpenFailed(FString FailedUrl);
-
-	/** Callback for when the player's tracks changed. */
-	void HandlePlayerTracksChanged();
+	/** Callback for when a media event occurred in the player. */
+	void HandlePlayerMediaEvent(EMediaEvent Event);
 
 private:
 
diff --git a/Engine/Source/Runtime/MovieScene/Private/MovieScene.cpp b/Engine/Source/Runtime/MovieScene/Private/MovieScene.cpp
--- a/Engine/Source/Runtime/MovieScene/Private/MovieScene.cpp
+++ b/Engine/Source/Runtime/MovieScene/Private/MovieScene.cpp
@@ -75,6 +75,12 @@ bool UMovieScene::RemoveSpawnable( const FGuid& Guid )
 
 	return bAnythingRemoved;
 }
+
+FMovieSceneSpawnable* UMovieScene::FindSpawnable( const TFunctionRef<bool(FMovieSceneSpawnable&)>& InPredicate )
+{
+	return Spawnables.FindByPredicate(InPredicate);
+}
+
 #endif //WITH_EDITOR
 
 
@@ -195,33 +201,47 @@ FMovieScenePossessable& UMovieScene::GetPossessable( const int32 Index )
 FText UMovieScene::GetObjectDisplayName(const FGuid& ObjectId)
 {
 #if WITH_EDITORONLY_DATA
-	FText& Result = ObjectsToDisplayNames.FindOrAdd(ObjectId.ToString());
+	FText* Result = ObjectsToDisplayNames.Find(ObjectId.ToString());
 
-	if (!Result.IsEmpty())
+	if (Result != nullptr)
 	{
-		return Result;
+		return *Result;
 	}
 
 	FMovieSceneSpawnable* Spawnable = FindSpawnable(ObjectId);
 
 	if (Spawnable != nullptr)
-		{
-		Result = FText::FromString(Spawnable->GetName());
-		return Result;
-		}
+	{
+		return FText::FromString(Spawnable->GetName());
+	}
 
 	FMovieScenePossessable* Possessable = FindPossessable(ObjectId);
 
 	if (Possessable != nullptr)
-		{
-		Result = FText::FromString(Possessable->GetName());
-		return Result;
+	{
+		return FText::FromString(Possessable->GetName());
 	}
 #endif
 	return FText::GetEmpty();
 }
 
 
+
+#if WITH_EDITORONLY_DATA
+void UMovieScene::SetObjectDisplayName(const FGuid& ObjectId, const FText& DisplayName)
+{
+	if (DisplayName.IsEmpty())
+	{
+		ObjectsToDisplayNames.Remove(ObjectId.ToString());
+	}
+	else
+	{
+		ObjectsToDisplayNames.Add(ObjectId.ToString(), DisplayName);
+	}
+}
+#endif
+
+
 TRange<float> UMovieScene::GetPlaybackRange() const
 {
 	check(PlaybackRange.HasLowerBound() && PlaybackRange.HasUpperBound());
diff --git a/Engine/Source/Runtime/MovieScene/Private/MovieSceneCommonHelpers.cpp b/Engine/Source/Runtime/MovieScene/Private/MovieSceneCommonHelpers.cpp
--- a/Engine/Source/Runtime/MovieScene/Private/MovieSceneCommonHelpers.cpp
+++ b/Engine/Source/Runtime/MovieScene/Private/MovieSceneCommonHelpers.cpp
@@ -143,20 +143,6 @@ USceneComponent* MovieSceneHelpers::SceneComponentFromRuntimeObject(UObject* Obj
 	return SceneComponent;
 }
 
-void MovieSceneHelpers::SetRuntimeObjectMobility(UObject* Object, EComponentMobility::Type ComponentMobility)
-{
-	USceneComponent* SceneComponent = SceneComponentFromRuntimeObject(Object);
-	if (SceneComponent)
-	{
-		if (SceneComponent->Mobility != ComponentMobility)
-		{
-			SceneComponent->Modify();
-
-			SceneComponent->SetMobility(ComponentMobility);
-		}
-	}
-}
-
 void MovieSceneHelpers::SetKeyInterpolation(FRichCurve& InCurve, FKeyHandle InKeyHandle, EMovieSceneKeyInterpolation InKeyInterpolation)
 {
 	switch (InKeyInterpolation)
diff --git a/Engine/Source/Runtime/MovieScene/Public/MovieScene.h b/Engine/Source/Runtime/MovieScene/Public/MovieScene.h
--- a/Engine/Source/Runtime/MovieScene/Public/MovieScene.h
+++ b/Engine/Source/Runtime/MovieScene/Public/MovieScene.h
@@ -105,6 +105,14 @@ public:
 	 */
 	bool RemoveSpawnable(const FGuid& Guid);
 
+	/**
+	 * Attempt to find a spawnable using some custom prodeicate
+	 *
+	 * @param InPredicate A predicate to test each spawnable against
+	 * @return Spawnable object that was found (or nullptr if not found).
+	 */
+	FMovieSceneSpawnable* FindSpawnable( const TFunctionRef<bool(FMovieSceneSpawnable&)>& InPredicate );
+
 #endif //WITH_EDITOR
 
 	/**
@@ -380,10 +388,7 @@ public:
 	 * @param ObjectId The object identifier.
 	 * @result The object's display name.
 	 */
-	void SetObjectDisplayName(const FGuid& ObjectId, const FText& DisplayName)
-	{
-		ObjectsToDisplayNames.Add(ObjectId.ToString(), DisplayName);
-	}
+	void SetObjectDisplayName(const FGuid& ObjectId, const FText& DisplayName);
 #endif
 
 	/**
diff --git a/Engine/Source/Runtime/MovieSceneCapture/MovieSceneCapture.Build.cs b/Engine/Source/Runtime/MovieSceneCapture/MovieSceneCapture.Build.cs
--- a/Engine/Source/Runtime/MovieSceneCapture/MovieSceneCapture.Build.cs
+++ b/Engine/Source/Runtime/MovieSceneCapture/MovieSceneCapture.Build.cs
@@ -33,6 +33,9 @@ public class MovieSceneCapture : ModuleRules
 				"Json",
 				"JsonUtilities",
 				"MovieScene",
+				"RenderCore",
+				"RHI",
+				"ShaderCore",
 				"Slate",
 				"SlateCore",
 			}
diff --git a/Engine/Source/Runtime/MovieSceneCapture/Private/ActiveMovieSceneCaptures.cpp b/Engine/Source/Runtime/MovieSceneCapture/Private/ActiveMovieSceneCaptures.cpp
--- a/Engine/Source/Runtime/MovieSceneCapture/Private/ActiveMovieSceneCaptures.cpp
+++ b/Engine/Source/Runtime/MovieSceneCapture/Private/ActiveMovieSceneCaptures.cpp
@@ -26,6 +26,26 @@ void FActiveMovieSceneCaptures::Remove(UMovieSceneCapture* Capture)
 	ActiveCaptures.Remove(Capture);
 }
 
+void FActiveMovieSceneCaptures::Tick(float DeltaSeconds)
+{
+	TArray<UMovieSceneCapture*> Captures = ActiveCaptures;
+	for (UMovieSceneCapture* Capture : Captures)
+	{
+		if (Capture->ShouldFinalize())
+		{
+			Capture->Finalize();
+		}
+		else
+		{
+			Capture->Tick(DeltaSeconds);
+			if (IMovieSceneCaptureProtocol* Protocol = Capture->GetCaptureProtocol())
+			{
+				Protocol->Tick();
+			}
+		}
+	}
+}
+
 void FActiveMovieSceneCaptures::Shutdown()
 {
 	TArray<UMovieSceneCapture*> ActiveCapturesCopy;
@@ -33,7 +53,7 @@ void FActiveMovieSceneCaptures::Shutdown()
 
 	for (auto* Obj : ActiveCapturesCopy)
 	{
-		Obj->Close();
+		Obj->Finalize();
 	}
 
 	Singleton.Reset();
diff --git a/Engine/Source/Runtime/MovieSceneCapture/Private/ActiveMovieSceneCaptures.h b/Engine/Source/Runtime/MovieSceneCapture/Private/ActiveMovieSceneCaptures.h
--- a/Engine/Source/Runtime/MovieSceneCapture/Private/ActiveMovieSceneCaptures.h
+++ b/Engine/Source/Runtime/MovieSceneCapture/Private/ActiveMovieSceneCaptures.h
@@ -4,7 +4,7 @@
 
 #include "MovieSceneCapture.h"
 
-class FActiveMovieSceneCaptures : public FGCObject
+class FActiveMovieSceneCaptures : public FGCObject, public FTickableGameObject
 {
 public:
 	/** Singleton access to an instance of this class */
@@ -27,6 +27,14 @@ private:
 	/** FGCObject interface */
 	virtual void AddReferencedObjects(FReferenceCollector& Collector) override;
 
+
+	/** FTickableGameObject interface */
+	virtual bool IsTickableInEditor() const override { return false; }
+	virtual bool IsTickable() const override { return ActiveCaptures.Num() != 0; }
+	virtual bool IsTickableWhenPaused() const override { return false; }
+	virtual TStatId GetStatId() const override { RETURN_QUICK_DECLARE_CYCLE_STAT(FActiveMovieSceneCaptures, STATGROUP_Tickables); }
+	virtual void Tick(float DeltaSeconds) override;
+
 private:
 	/** Private construction to enforce singleton use */
 	FActiveMovieSceneCaptures() {}
diff --git a/Engine/Source/Runtime/MovieSceneCapture/Private/AutomatedLevelSequenceCapture.cpp b/Engine/Source/Runtime/MovieSceneCapture/Private/AutomatedLevelSequenceCapture.cpp
--- a/Engine/Source/Runtime/MovieSceneCapture/Private/AutomatedLevelSequenceCapture.cpp
+++ b/Engine/Source/Runtime/MovieSceneCapture/Private/AutomatedLevelSequenceCapture.cpp
@@ -5,6 +5,7 @@
 #include "LevelSequencePlayer.h"
 #include "AutomatedLevelSequenceCapture.h"
 #include "ErrorCodes.h"
+#include "SceneViewport.h"
 #include "ActiveMovieSceneCaptures.h"
 
 UAutomatedLevelSequenceCapture::UAutomatedLevelSequenceCapture(const FObjectInitializer& Init)
@@ -35,13 +36,7 @@ void UAutomatedLevelSequenceCapture::SetLevelSequenceAsset(FString AssetPath)
 	LevelSequenceAsset = MoveTemp(AssetPath);
 }
 
-void UAutomatedLevelSequenceCapture::SetLevelSequenceActor(ALevelSequenceActor* InActor)
-{
-	LevelSequenceActor = InActor;
-	LevelSequenceActorId = LevelSequenceActor.GetUniqueID().GetGuid();
-}
-
-void UAutomatedLevelSequenceCapture::Initialize(TWeakPtr<FSceneViewport> InViewport)
+void UAutomatedLevelSequenceCapture::Initialize(TSharedPtr<FSceneViewport> InViewport, int32 PIEInstance)
 {
 	// Apply command-line overrides from parent class first. This needs to be called before setting up the capture strategy with the desired frame rate.
 	Super::Initialize(InViewport);
@@ -52,7 +47,6 @@ void UAutomatedLevelSequenceCapture::Initialize(TWeakPtr<FSceneViewport> InViewp
 		if( FParse::Value( FCommandLine::Get(), TEXT( "-LevelSequence=" ), LevelSequenceAssetPath ) )
 		{
 			LevelSequenceAsset.SetPath( LevelSequenceAssetPath );
-			LevelSequenceActorId = FGuid();
 		}
 
 		int32 StartFrameOverride;
@@ -82,9 +76,6 @@ void UAutomatedLevelSequenceCapture::Initialize(TWeakPtr<FSceneViewport> InViewp
 		}
 	}
 
-	// Ensure the LevelSequence is up to date with the LevelSequenceActorId (Level sequence is only there for a nice UI)
-	LevelSequenceActor = FUniqueObjectGuid(LevelSequenceActorId);
-
 	ALevelSequenceActor* Actor = LevelSequenceActor.Get();
 
 	// If we don't have a valid actor, attempt to find a level sequence actor in the world that references this asset
@@ -95,7 +86,7 @@ void UAutomatedLevelSequenceCapture::Initialize(TWeakPtr<FSceneViewport> InViewp
 			ULevelSequence* Asset = Cast<ULevelSequence>( LevelSequenceAsset.TryLoad() );
 			if( Asset != nullptr )
 			{
-				for( auto It = TActorIterator<ALevelSequenceActor>( InViewport.Pin()->GetClient()->GetWorld() ); It; ++It )
+				for( auto It = TActorIterator<ALevelSequenceActor>( InViewport->GetClient()->GetWorld() ); It; ++It )
 				{
 					if( It->LevelSequence == LevelSequenceAsset )
 					{
@@ -110,6 +101,25 @@ void UAutomatedLevelSequenceCapture::Initialize(TWeakPtr<FSceneViewport> InViewp
 		}
 	}
 
+	if (!Actor)
+	{
+		ULevelSequence* Asset = Cast<ULevelSequence>(LevelSequenceAsset.TryLoad());
+		if (Asset)
+		{
+			// Spawn a new actor
+			Actor = GWorld->SpawnActor<ALevelSequenceActor>();
+			Actor->SetSequence(Asset);
+			// Ensure it doesn't loop (-1 is indefinite)
+			Actor->PlaybackSettings.LoopCount = 0;
+		
+			LevelSequenceActor = Actor;
+		}
+		else
+		{
+			//FPlatformMisc::RequestExit(FMovieSceneCaptureExitCodes::AssetNotFound);
+		}
+	}
+
 	if (Actor)
 	{
 		// Make sure we're not playing yet (in case AutoPlay was called from BeginPlay)
@@ -118,20 +128,15 @@ void UAutomatedLevelSequenceCapture::Initialize(TWeakPtr<FSceneViewport> InViewp
 			Actor->SequencePlayer->Stop();
 		}
 		Actor->bAutoPlay = false;
-
-		// Bind to the event so we know when to capture a frame
-		if( Actor->SequencePlayer != nullptr )
-		{
-			OnPlayerUpdatedBinding = Actor->SequencePlayer->OnSequenceUpdated().AddUObject( this, &UAutomatedLevelSequenceCapture::SequenceUpdated );
-		}
 	}
 
-	CaptureState = ELevelSequenceCaptureState::DelayBeforeWarmUp;
+	CaptureState = ELevelSequenceCaptureState::Setup;
 	RemainingDelaySeconds = FMath::Max( 0.0f, DelayBeforeWarmUp );
 	CaptureStrategy = MakeShareable(new FFixedTimeStepCaptureStrategy(Settings.FrameRate));
 }
 
 
+
 void UAutomatedLevelSequenceCapture::SetupFrameRange()
 {
 	ALevelSequenceActor* Actor = LevelSequenceActor.Get();
@@ -166,10 +171,14 @@ void UAutomatedLevelSequenceCapture::SetupFrameRange()
 				if( bUseCustomEndFrame )
 				{
 					PlaybackEndFrame = FMath::Max( PlaybackStartFrame, Settings.bUseRelativeFrameNumbers ? ( SequenceEndFrame + EndFrame ) : EndFrame );
-				}
 
-				// We always add 1 to the number of frames we want to capture, because we want to capture both the start and end frames (which if the play range is 0, would still yield a single frame)
-				this->FrameCount = ( PlaybackEndFrame - PlaybackStartFrame ) + 1;
+					// We always add 1 to the number of frames we want to capture, because we want to capture both the start and end frames (which if the play range is 0, would still yield a single frame)
+					this->FrameCount = ( PlaybackEndFrame - PlaybackStartFrame ) + 1;
+				}
+				else
+				{
+					FrameCount = 0;
+				}
 
 				RemainingWarmUpFrames = FMath::Max( WarmUpFrameCount, 0 );
 				if( RemainingWarmUpFrames > 0 )
@@ -190,123 +199,83 @@ void UAutomatedLevelSequenceCapture::SetupFrameRange()
 
 void UAutomatedLevelSequenceCapture::Tick(float DeltaSeconds)
 {
-	const bool bAnyFramesToCapture = OutstandingFrameCount > 0;
-	Super::Tick(DeltaSeconds);
-
 	ALevelSequenceActor* Actor = LevelSequenceActor.Get();
 
-	if (!Actor)
+	if (!Actor || !Actor->SequencePlayer)
 	{
-		ULevelSequence* Asset = Cast<ULevelSequence>(LevelSequenceAsset.TryLoad());
-		if (Asset)
-		{
-			// Spawn a new actor
-			Actor = GWorld->SpawnActor<ALevelSequenceActor>();
-			Actor->SetSequence(Asset);
-			// Ensure it doesn't loop (-1 is indefinite)
-			Actor->PlaybackSettings.LoopCount = 0;
-			
-			LevelSequenceActor = Actor;
-		}
-		else
-		{
-			FPlatformMisc::RequestExit(FMovieSceneCaptureExitCodes::AssetNotFound);
-		}
+		return;
 	}
 
-	if (Actor && Actor->SequencePlayer)
+	// Setup the automated capture
+	if (CaptureState == ELevelSequenceCaptureState::Setup)
 	{
-		// First off we'll stage the sequence.  This just means we'll play the first frame of the animation
-		// and then pause it immediately
-		if( CaptureState == ELevelSequenceCaptureState::Staging )
-		{
-			Actor->SequencePlayer->Play();
-			Actor->SequencePlayer->Pause();
+		SetupFrameRange();
+		
+		// Bind to the event so we know when to capture a frame
+		OnPlayerUpdatedBinding = Actor->SequencePlayer->OnSequenceUpdated().AddUObject( this, &UAutomatedLevelSequenceCapture::SequenceUpdated );
 
-			CaptureState = ELevelSequenceCaptureState::DelayBeforeWarmUp;
-		}
+		CaptureState = ELevelSequenceCaptureState::DelayBeforeWarmUp;
+	}
+
+	bool bShouldPlaySequence = false;
 
-		// Then we'll just wait a little bit.  We'll delay the specified number of seconds before capturing to allow any
-		// textures to stream in or post processing effects to settle.
-		if( CaptureState == ELevelSequenceCaptureState::DelayBeforeWarmUp )
+	// Then we'll just wait a little bit.  We'll delay the specified number of seconds before capturing to allow any
+	// textures to stream in or post processing effects to settle.
+	if( CaptureState == ELevelSequenceCaptureState::DelayBeforeWarmUp )
+	{
+		RemainingDelaySeconds -= DeltaSeconds;
+		if( RemainingDelaySeconds <= 0.0f )
 		{
-			RemainingDelaySeconds -= DeltaSeconds;
-			if( RemainingDelaySeconds <= 0.0f )
-			{
-				RemainingDelaySeconds = 0.0f;
+			RemainingDelaySeconds = 0.0f;
 
-				// Start warming up.  Even if we're not capturing yet, this will make sure we're rendering at a
-				// fixed frame rate.
-				StartWarmup();
+			// Start warming up.  Even if we're not capturing yet, this will make sure we're rendering at a
+			// fixed frame rate.
+			StartWarmup();
 
-				// Wait a frame to go by after we've set the fixed time step, so that the animation starts
-				// playback at a consistent time
-				CaptureState = ELevelSequenceCaptureState::ReadyToWarmUp;
-			}
-		}
-		else if( CaptureState == ELevelSequenceCaptureState::ReadyToWarmUp )
-		{
-			Actor->SequencePlayer->Play();
-			CaptureState = ELevelSequenceCaptureState::WarmingUp;
+			// Wait a frame to go by after we've set the fixed time step, so that the animation starts
+			// playback at a consistent time
+			CaptureState = ELevelSequenceCaptureState::ReadyToWarmUp;
 		}
+	}
+	else if( CaptureState == ELevelSequenceCaptureState::ReadyToWarmUp )
+	{
+		Actor->SequencePlayer->Play();
+		CaptureState = ELevelSequenceCaptureState::WarmingUp;
+	}
 
 
-		if( CaptureState == ELevelSequenceCaptureState::WarmingUp )
+	if( CaptureState == ELevelSequenceCaptureState::WarmingUp )
+	{
+		// Count down our warm up frames
+		if( RemainingWarmUpFrames == 0 )
 		{
-			// Count down our warm up frames
-			if( RemainingWarmUpFrames == 0 )
-			{
-				CaptureState = ELevelSequenceCaptureState::FinishedWarmUp;
+			CaptureState = ELevelSequenceCaptureState::FinishedWarmUp;
 
-				// It's time to start capturing!
-				StartCapture();
-			}
-			else
-			{
-				// Not ready to capture just yet
-				--RemainingWarmUpFrames;
-			}
+			// It's time to start capturing!
+			StartCapture();
+			CaptureThisFrame(LastSequenceUpdateDelta);
 		}
-
-
-		if( bAnyFramesToCapture && OutstandingFrameCount == 0 )
+		else
 		{
-			// If we hit this, then we've rendered out the last frame and can stop!
-			StopCapture();
+			// Not ready to capture just yet
+			--RemainingWarmUpFrames;
 		}
 	}
-}
 
-void UAutomatedLevelSequenceCapture::OnCaptureStopped()
-{
-	ALevelSequenceActor* Actor = LevelSequenceActor.Get();
-	if( Actor != nullptr )
+	if( bCapturing && !Actor->SequencePlayer->IsPlaying() )
 	{
-		if( Actor->SequencePlayer != nullptr )
-		{
-			Actor->SequencePlayer->OnSequenceUpdated().Remove( OnPlayerUpdatedBinding );
-		}
+		Actor->SequencePlayer->OnSequenceUpdated().Remove( OnPlayerUpdatedBinding );
+		FinalizeWhenReady();
 	}
-
-	Close();
-	FPlatformMisc::RequestExit(0);
 }
 
 void UAutomatedLevelSequenceCapture::SequenceUpdated(const ULevelSequencePlayer& Player, float CurrentTime, float PreviousTime)
 {
-	if( bCapturing )
+	LastSequenceUpdateDelta = CurrentTime - PreviousTime;
+	if (bCapturing)
 	{
-		// Save the previous rendered frame
-		CaptureFrame( LastFrameDelta );
-
-		// Get ready to capture the next rendered frame
-		PrepareForScreenshot();
+		CaptureThisFrame(LastSequenceUpdateDelta);
 	}
 }
 
-void UAutomatedLevelSequenceCapture::PostEditChangeProperty(struct FPropertyChangedEvent& PropertyChangedEvent)
-{
-	LevelSequenceActorId = LevelSequenceActor.GetUniqueID().GetGuid();
-}
-
 #endif
\ No newline at end of file
diff --git a/Engine/Source/Runtime/MovieSceneCapture/Private/CompositionGraphCaptureProtocol.cpp b/Engine/Source/Runtime/MovieSceneCapture/Private/CompositionGraphCaptureProtocol.cpp
new file mode 100644
--- /dev/null
+++ b/Engine/Source/Runtime/MovieSceneCapture/Private/CompositionGraphCaptureProtocol.cpp
@@ -0,0 +1,158 @@
+// Copyright 1998-2015 Epic Games, Inc. All Rights Reserved.
+
+#include "MovieSceneCapturePCH.h"
+#include "CompositionGraphCaptureProtocol.h"
+#include "SceneViewExtension.h"
+#include "BufferVisualizationData.h"
+#include "SceneViewport.h"
+
+struct FSceneViewExtension : ISceneViewExtension
+{
+	FSceneViewExtension(const TArray<FString>& InRenderPasses, bool bInCaptureFramesInHDR, UMaterialInterface* InPostProcessingMaterial)
+		: RenderPasses(InRenderPasses)
+	{
+		PostProcessingMaterial = InPostProcessingMaterial;
+		bCaptureFramesInHDR = bInCaptureFramesInHDR;
+
+		CVarDumpFrames = IConsoleManager::Get().FindConsoleVariable(TEXT("r.BufferVisualizationDumpFrames"));
+		CVarDumpFramesAsHDR = IConsoleManager::Get().FindConsoleVariable(TEXT("r.BufferVisualizationDumpFramesAsHDR"));
+
+		Disable();
+	}
+
+	virtual ~FSceneViewExtension()
+	{
+		Disable();
+	}
+
+	bool IsEnabled() const
+	{
+		return bNeedsCapture;
+	}
+
+	void Enable(FString&& InFilename)
+	{
+		OutputFilename = MoveTemp(InFilename);
+
+		bNeedsCapture = true;
+		RestoreDumpHDR = CVarDumpFramesAsHDR->GetInt();
+		CVarDumpFramesAsHDR->Set(bCaptureFramesInHDR);
+		CVarDumpFrames->Set(1);
+	}
+
+	void Disable()
+	{
+		if (bNeedsCapture)
+		{
+			bNeedsCapture = false;
+			CVarDumpFramesAsHDR->Set(RestoreDumpHDR);
+			CVarDumpFrames->Set(0);
+		}
+	}
+
+	virtual void SetupView(FSceneViewFamily& InViewFamily, FSceneView& InView)
+	{
+		if (!bNeedsCapture)
+		{
+			return;
+		}
+
+		InView.FinalPostProcessSettings.bBufferVisualizationDumpRequired = true;
+		InView.FinalPostProcessSettings.BufferVisualizationOverviewMaterials.Empty();
+		InView.FinalPostProcessSettings.BufferVisualizationDumpBaseFilename = MoveTemp(OutputFilename);
+
+		struct FIterator
+		{
+			FFinalPostProcessSettings& FinalPostProcessSettings;
+			const TArray<FString>& RenderPasses;
+
+			FIterator(FFinalPostProcessSettings& InFinalPostProcessSettings, const TArray<FString>& InRenderPasses)
+				: FinalPostProcessSettings(InFinalPostProcessSettings), RenderPasses(InRenderPasses)
+			{}
+
+			void ProcessValue(const FString& InName, UMaterial* Material, const FText&)
+			{
+				if (!RenderPasses.Num() || RenderPasses.Contains(InName))
+				{
+					FinalPostProcessSettings.BufferVisualizationOverviewMaterials.Add(Material);
+				}
+			}
+		} Iterator(InView.FinalPostProcessSettings, RenderPasses);
+		GetBufferVisualizationData().IterateOverAvailableMaterials(Iterator);
+
+		if (PostProcessingMaterial)
+		{
+			FWeightedBlendable Blendable(1.f, PostProcessingMaterial);
+			PostProcessingMaterial->OverrideBlendableSettings(InView, 1.f);
+		}
+
+		bNeedsCapture = false;
+	}
+
+	virtual void SetupViewFamily(FSceneViewFamily& InViewFamily) {}
+	virtual void BeginRenderViewFamily(FSceneViewFamily& InViewFamily) {}
+	virtual void PreRenderViewFamily_RenderThread(FRHICommandListImmediate& RHICmdList, FSceneViewFamily& InViewFamily) {}
+	virtual void PreRenderView_RenderThread(FRHICommandListImmediate& RHICmdList, FSceneView& InView) {}
+
+private:
+	const TArray<FString>& RenderPasses;
+
+	bool bCaptureFramesInHDR;
+
+	UMaterialInterface* PostProcessingMaterial;
+
+	bool bNeedsCapture;
+	FString OutputFilename;
+
+	IConsoleVariable* CVarDumpFrames;
+	IConsoleVariable* CVarDumpFramesAsHDR;
+
+	int32 RestoreDumpHDR;
+};
+
+bool FCompositionGraphCaptureProtocol::Initialize(const FCaptureProtocolInitSettings& InSettings, const ICaptureProtocolHost& Host)
+{
+	SceneViewport = InSettings.SceneViewport;
+
+	bool bCaptureFramesInHDR = false;
+
+	UMaterialInterface* PostProcessingMaterial = nullptr;
+	UCompositionGraphCaptureSettings* ProtocolSettings = CastChecked<UCompositionGraphCaptureSettings>(InSettings.ProtocolSettings);
+	if (ProtocolSettings)
+	{
+		RenderPasses = ProtocolSettings->IncludeRenderPasses.Value;
+		bCaptureFramesInHDR = ProtocolSettings->bCaptureFramesInHDR;
+		PostProcessingMaterial = Cast<UMaterialInterface>(ProtocolSettings->PostProcessingMaterial.TryLoad());
+	}
+
+	ViewExtension = MakeShareable(new FSceneViewExtension(RenderPasses, bCaptureFramesInHDR, PostProcessingMaterial));
+
+	return true;
+}
+
+void FCompositionGraphCaptureProtocol::Finalize()
+{
+	ViewExtension->Disable();
+	GEngine->ViewExtensions.Remove(ViewExtension);
+}
+
+void FCompositionGraphCaptureProtocol::CaptureFrame(const FFrameMetrics& FrameMetrics, const ICaptureProtocolHost& Host)
+{
+	ViewExtension->Enable(Host.GenerateFilename(FrameMetrics, TEXT("")));
+
+	GEngine->ViewExtensions.AddUnique(ViewExtension);
+}
+
+bool FCompositionGraphCaptureProtocol::HasFinishedProcessing() const
+{
+	return !ViewExtension->IsEnabled();
+}
+
+void FCompositionGraphCaptureProtocol::Tick()
+{
+	if (!ViewExtension->IsEnabled())
+	{
+		ViewExtension->Disable();
+		GEngine->ViewExtensions.Remove(ViewExtension);
+	}
+}
\ No newline at end of file
