Copying //UE4/Dev-Sequencer to Dev-Main (//UE4/Dev-Main)

diff --git a/Engine/Source/Runtime/MovieSceneCapture/Private/FrameGrabber.cpp b/Engine/Source/Runtime/MovieSceneCapture/Private/FrameGrabber.cpp
new file mode 100644
--- /dev/null
+++ b/Engine/Source/Runtime/MovieSceneCapture/Private/FrameGrabber.cpp
@@ -0,0 +1,374 @@
+// Copyright 1998-2015 Epic Games, Inc. All Rights Reserved.
+
+#include "MovieSceneCapturePCH.h"
+#include "FrameGrabber.h"
+#include "SceneViewport.h"
+#include "ScreenRendering.h"
+#include "RenderCore.h"
+#include "RHIStaticStates.h"
+#include "RendererInterface.h"
+
+FViewportSurfaceReader::FViewportSurfaceReader(EPixelFormat InPixelFormat, FIntPoint InBufferSize)
+{
+	AvailableEvent = nullptr;
+	ReadbackTexture = nullptr;
+	PixelFormat = InPixelFormat;
+
+	Resize(InBufferSize.X, InBufferSize.Y);
+}
+
+FViewportSurfaceReader::~FViewportSurfaceReader()
+{
+	BlockUntilAvailable();
+
+	ReadbackTexture.SafeRelease();
+}
+
+void FViewportSurfaceReader::Initialize()
+{
+	check(!AvailableEvent);
+	AvailableEvent = FPlatformProcess::GetSynchEventFromPool();
+}
+
+void FViewportSurfaceReader::Resize(uint32 Width, uint32 Height)
+{
+	ReadbackTexture.SafeRelease();
+
+	ENQUEUE_UNIQUE_RENDER_COMMAND_THREEPARAMETER(
+		CreateCaptureFrameTexture,
+		int32, InWidth, Width,
+		int32, InHeight, Height,
+		FViewportSurfaceReader*, This, this,
+	{
+		FRHIResourceCreateInfo CreateInfo;
+
+		This->ReadbackTexture = RHICreateTexture2D(
+			InWidth,
+			InHeight,
+			This->PixelFormat,
+			1,
+			1,
+			TexCreate_CPUReadback,
+			CreateInfo
+			);
+	});
+}
+
+void FViewportSurfaceReader::BlockUntilAvailable()
+{
+	if (AvailableEvent)
+	{
+		AvailableEvent->Wait(~0);
+
+		FPlatformProcess::ReturnSynchEventToPool(AvailableEvent);
+		AvailableEvent = nullptr;
+	}
+}
+
+void FViewportSurfaceReader::ResolveRenderTarget(const FViewportRHIRef& ViewportRHI, TFunction<void(FColor*)> Callback)
+{
+	static const FName RendererModuleName( "Renderer" );
+	IRendererModule* RendererModule = &FModuleManager::GetModuleChecked<IRendererModule>(RendererModuleName);
+
+	auto RenderCommand = [=](FRHICommandListImmediate& RHICmdList){
+
+		const FIntPoint TargetSize(ReadbackTexture->GetSizeX(), ReadbackTexture->GetSizeY());
+
+		FPooledRenderTargetDesc OutputDesc = FPooledRenderTargetDesc::Create2DDesc(
+			TargetSize,
+			ReadbackTexture->GetFormat(),
+			FClearValueBinding::None,
+			TexCreate_None,
+			TexCreate_RenderTargetable,
+			false);
+
+		TRefCountPtr<IPooledRenderTarget> ResampleTexturePooledRenderTarget;
+		RendererModule->RenderTargetPoolFindFreeElement(RHICmdList, OutputDesc, ResampleTexturePooledRenderTarget, TEXT("ResampleTexture"));
+		check(ResampleTexturePooledRenderTarget);
+
+		const FSceneRenderTargetItem& DestRenderTarget = ResampleTexturePooledRenderTarget->GetRenderTargetItem();
+
+		SetRenderTarget(RHICmdList, DestRenderTarget.TargetableTexture, FTextureRHIRef());
+
+		RHICmdList.SetViewport(0, 0, 0.0f, TargetSize.X, TargetSize.Y, 1.0f);
+
+		RHICmdList.SetBlendState(TStaticBlendState<>::GetRHI());
+		RHICmdList.SetRasterizerState(TStaticRasterizerState<>::GetRHI());
+		RHICmdList.SetDepthStencilState(TStaticDepthStencilState<false,CF_Always>::GetRHI());
+
+
+		const ERHIFeatureLevel::Type FeatureLevel = GMaxRHIFeatureLevel;
+		
+		TShaderMap<FGlobalShaderType>* ShaderMap = GetGlobalShaderMap(FeatureLevel);
+		TShaderMapRef<FScreenVS> VertexShader(ShaderMap);
+		TShaderMapRef<FScreenPS> PixelShader(ShaderMap);
+
+		static FGlobalBoundShaderState BoundShaderState;
+		SetGlobalBoundShaderState(RHICmdList, FeatureLevel, BoundShaderState, RendererModule->GetFilterVertexDeclaration().VertexDeclarationRHI, *VertexShader, *PixelShader);
+
+		FTexture2DRHIRef SourceBackBuffer = RHICmdList.GetViewportBackBuffer(ViewportRHI);
+
+		if (TargetSize.X != SourceBackBuffer->GetSizeX() || TargetSize.Y != SourceBackBuffer->GetSizeY())
+		{
+			PixelShader->SetParameters(RHICmdList, TStaticSamplerState<SF_Bilinear>::GetRHI(), SourceBackBuffer);
+		}
+		else
+		{
+			PixelShader->SetParameters(RHICmdList, TStaticSamplerState<SF_Point>::GetRHI(), SourceBackBuffer);
+		}
+
+		float U = float(CaptureRect.Min.X) / float(SourceBackBuffer->GetSizeX());
+		float V = float(CaptureRect.Min.Y) / float(SourceBackBuffer->GetSizeY());
+		float SizeU = float(CaptureRect.Max.X) / float(SourceBackBuffer->GetSizeX()) - U;
+		float SizeV = float(CaptureRect.Max.Y) / float(SourceBackBuffer->GetSizeY()) - V;
+
+		RendererModule->DrawRectangle(
+			RHICmdList,
+			0, 0,									// Dest X, Y
+			TargetSize.X,							// Dest Width
+			TargetSize.Y,							// Dest Height
+			U, V,									// Source U, V
+			1, 1,									// Source USize, VSize
+			CaptureRect.Max - CaptureRect.Min,		// Target buffer size
+			FIntPoint(1, 1),						// Source texture size
+			*VertexShader,
+			EDRF_Default);
+
+		// Asynchronously copy render target from GPU to CPU
+		const bool bKeepOriginalSurface = false;
+		RHICmdList.CopyToResolveTarget(
+			DestRenderTarget.TargetableTexture,
+			ReadbackTexture,
+			bKeepOriginalSurface,
+			FResolveParams());
+
+		void* ColorDataBuffer = nullptr;
+
+		int32 UnusedWidth = 0, UnusedHeight = 0;
+		RHICmdList.MapStagingSurface(ReadbackTexture, ColorDataBuffer, UnusedWidth, UnusedHeight);
+
+		Callback((FColor*)ColorDataBuffer);
+
+		RHICmdList.UnmapStagingSurface(ReadbackTexture);
+		AvailableEvent->Trigger();
+	};
+
+	ENQUEUE_UNIQUE_RENDER_COMMAND_ONEPARAMETER(
+		ResolveCaptureFrameTexture,
+		TFunction<void(FRHICommandListImmediate&)>, InRenderCommand, RenderCommand,
+	{
+		InRenderCommand(RHICmdList);
+	});
+}
+
+FFrameGrabber::FFrameGrabber(TSharedRef<FSceneViewport> Viewport, FIntPoint DesiredBufferSize, EPixelFormat InPixelFormat, uint32 NumSurfaces)
+{
+	State = EFrameGrabberState::Inactive;
+
+	TargetSize = DesiredBufferSize;
+
+	CurrentFrameIndex = 0;
+
+	check(NumSurfaces != 0);
+
+	FIntRect CaptureRect(0,0,Viewport->GetSize().X, Viewport->GetSize().Y);
+
+	// Set up the capture rectangle
+	TSharedPtr<SViewport> ViewportWidget = Viewport->GetViewportWidget().Pin();
+	if (ViewportWidget.IsValid())
+	{
+		TSharedPtr<SWindow> Window = FSlateApplication::Get().FindWidgetWindow(ViewportWidget.ToSharedRef());
+		if (Window.IsValid())
+		{
+			CaptureWindow = Window;
+			FGeometry InnerWindowGeometry = Window->GetWindowGeometryInWindow();
+			
+			// Find the widget path relative to the window
+			FArrangedChildren JustWindow(EVisibility::Visible);
+			JustWindow.AddWidget(FArrangedWidget(Window.ToSharedRef(), InnerWindowGeometry));
+
+			FWidgetPath WidgetPath(Window.ToSharedRef(), JustWindow);
+			if (WidgetPath.ExtendPathTo(FWidgetMatcher(ViewportWidget.ToSharedRef()), EVisibility::Visible))
+			{
+				FArrangedWidget ArrangedWidget = WidgetPath.FindArrangedWidget(ViewportWidget.ToSharedRef()).Get(FArrangedWidget::NullWidget);
+
+				FVector2D Position = ArrangedWidget.Geometry.AbsolutePosition;
+				FVector2D Size = ArrangedWidget.Geometry.GetDrawSize();
+
+				CaptureRect = FIntRect(
+					Position.X,
+					Position.Y,
+					Position.X + Size.X,
+					Position.Y + Size.Y);
+			}
+		}
+	}
+
+	// This can never be reallocated
+	Surfaces.Reserve(NumSurfaces);
+	for (uint32 Index = 0; Index < NumSurfaces; ++Index)
+	{
+		Surfaces.Emplace(InPixelFormat, DesiredBufferSize);
+		Surfaces.Last().Surface.SetCaptureRect(CaptureRect);
+	}
+
+	// Ensure textures are setup
+	FlushRenderingCommands();
+}
+
+FFrameGrabber::~FFrameGrabber()
+{
+	if (OnWindowRendered.IsValid())
+	{
+		FSlateApplication::Get().GetRenderer()->OnSlateWindowRendered().Remove(OnWindowRendered);
+	}
+}
+
+void FFrameGrabber::StartCapturingFrames()
+{
+	if (!ensure(State == EFrameGrabberState::Inactive))
+	{
+		return;
+	}
+
+	State = EFrameGrabberState::Active;
+	OnWindowRendered = FSlateApplication::Get().GetRenderer()->OnSlateWindowRendered().AddRaw(this, &FFrameGrabber::OnSlateWindowRendered);
+}
+
+void FFrameGrabber::CaptureThisFrame(FFramePayloadPtr Payload)
+{
+	if (!ensure(State == EFrameGrabberState::Active))
+	{
+		return;
+	}
+
+	OutstandingFrameCount.Increment();
+
+	// Issue a rendering command to ensure we capture the frame that is currently being set up
+	ENQUEUE_UNIQUE_RENDER_COMMAND_TWOPARAMETER(
+		IssueCaptureFrameCommand,
+		FFrameGrabber* , This, this,
+		FFramePayloadPtr, Payload, Payload,
+	{
+		FScopeLock Lock(&This->PendingFramePayloadsMutex);
+		This->PendingFramePayloads.Add(Payload);
+	});
+}
+
+void FFrameGrabber::StopCapturingFrames()
+{
+	if (!ensure(State == EFrameGrabberState::Active))
+	{
+		return;
+	}
+
+	State = EFrameGrabberState::PendingShutdown;
+}
+
+void FFrameGrabber::Shutdown()
+{
+	State = EFrameGrabberState::Inactive;
+
+	for (FResolveSurface& Surface : Surfaces)
+	{
+		Surface.Surface.BlockUntilAvailable();
+	}
+
+	FSlateApplication::Get().GetRenderer()->OnSlateWindowRendered().Remove(OnWindowRendered);
+	OnWindowRendered = FDelegateHandle();
+}
+
+bool FFrameGrabber::HasOutstandingFrames() const
+{
+	FScopeLock Lock(&CapturedFramesMutex);
+
+	// Check whether we have any outstanding frames while we have the array locked, to prevent race condition
+	return OutstandingFrameCount.GetValue() || CapturedFrames.Num();
+}
+
+TArray<FCapturedFrameData> FFrameGrabber::GetCapturedFrames()
+{
+	TArray<FCapturedFrameData> ReturnFrames;
+
+	bool bShouldStop = false;
+	{
+		FScopeLock Lock(&CapturedFramesMutex);
+		Swap(ReturnFrames, CapturedFrames);
+		CapturedFrames.Reset();
+
+		// Check whether we have any outstanding frames while we have the array locked, to prevent race condition
+		if (State == EFrameGrabberState::PendingShutdown)
+		{
+			bShouldStop = OutstandingFrameCount.GetValue() == 0;
+		}
+	}
+
+	if (bShouldStop)
+	{
+		Shutdown();
+	}
+
+	return ReturnFrames;
+}
+
+void FFrameGrabber::OnSlateWindowRendered( SWindow& SlateWindow, void* ViewportRHIPtr )
+{
+	// We only care about our own Slate window
+	TSharedPtr<SWindow> Window = CaptureWindow.Pin();
+	if (Window.Get() != &SlateWindow)
+	{
+		return;
+	}
+
+	FFramePayloadPtr Payload;
+
+	{
+		FScopeLock Lock(&PendingFramePayloadsMutex);
+		if (!PendingFramePayloads.Num())
+		{
+			// No frames to capture
+			return;
+		}
+		Payload = PendingFramePayloads[0];
+		PendingFramePayloads.RemoveAt(0, 1, false);
+	}
+
+	const int32 ThisCaptureIndex = CurrentFrameIndex;
+
+	FResolveSurface* ThisFrameTarget = &Surfaces[ThisCaptureIndex];
+	ThisFrameTarget->Surface.BlockUntilAvailable();
+
+	ThisFrameTarget->Surface.Initialize();
+	ThisFrameTarget->Payload = Payload;
+
+	const FViewportRHIRef* RHIViewport = (const FViewportRHIRef*)ViewportRHIPtr;
+
+	Surfaces[ThisCaptureIndex].Surface.ResolveRenderTarget(*RHIViewport, [=](FColor* ColorBuffer){
+		// Handle the frame
+		OnFrameReady(ThisCaptureIndex, ColorBuffer);
+	});
+
+	CurrentFrameIndex = (CurrentFrameIndex+1) % Surfaces.Num();
+}
+
+void FFrameGrabber::OnFrameReady(int32 BufferIndex, FColor* ColorBuffer)
+{
+	if (!ensure(ColorBuffer))
+	{
+		OutstandingFrameCount.Decrement();
+		return;
+	}
+	
+	const FResolveSurface& Surface = Surfaces[BufferIndex];
+
+	FCapturedFrameData ResolvedFrameData(TargetSize, Surface.Payload);
+
+	ResolvedFrameData.ColorBuffer.Reserve(TargetSize.X * TargetSize.Y);
+	ResolvedFrameData.ColorBuffer.Insert(ColorBuffer, TargetSize.X * TargetSize.Y, 0);
+
+	{
+		FScopeLock Lock(&CapturedFramesMutex);
+		CapturedFrames.Add(MoveTemp(ResolvedFrameData));
+	}
+	OutstandingFrameCount.Decrement();
+}
\ No newline at end of file
diff --git a/Engine/Source/Runtime/MovieSceneCapture/Private/FrameGrabberProtocol.cpp b/Engine/Source/Runtime/MovieSceneCapture/Private/FrameGrabberProtocol.cpp
new file mode 100644
--- /dev/null
+++ b/Engine/Source/Runtime/MovieSceneCapture/Private/FrameGrabberProtocol.cpp
@@ -0,0 +1,47 @@
+// Copyright 1998-2015 Epic Games, Inc. All Rights Reserved.
+
+#include "MovieSceneCapturePCH.h"
+#include "FrameGrabberProtocol.h"
+
+bool FFrameGrabberProtocol::HasFinishedProcessing() const
+{
+	return !FrameGrabber->HasOutstandingFrames();
+}
+
+bool FFrameGrabberProtocol::Initialize(const FCaptureProtocolInitSettings& InSettings, const ICaptureProtocolHost& Host)
+{
+	EPixelFormat PixelFormat = PF_B8G8R8A8;
+	uint32 RingBufferSize = 3;
+
+	if (UFrameGrabberProtocolSettings* Settings = Cast<UFrameGrabberProtocolSettings>(InSettings.ProtocolSettings))
+	{
+		PixelFormat = Settings->DesiredPixelFormat;
+		RingBufferSize = Settings->RingBufferSize;
+	}
+
+	// We'll use our own grabber to capture the entire viewport
+	FrameGrabber.Reset(new FFrameGrabber(InSettings.SceneViewport.ToSharedRef(), InSettings.DesiredSize, PixelFormat, RingBufferSize));
+	FrameGrabber->StartCapturingFrames();
+	return true;
+}
+
+void FFrameGrabberProtocol::CaptureFrame(const FFrameMetrics& FrameMetrics, const ICaptureProtocolHost& Host)
+{
+	FrameGrabber->CaptureThisFrame(GetFramePayload(FrameMetrics, Host));
+}
+
+void FFrameGrabberProtocol::Tick()
+{
+	TArray<FCapturedFrameData> CapturedFrames = FrameGrabber->GetCapturedFrames();
+
+	for (FCapturedFrameData& Frame : CapturedFrames)
+	{
+		ProcessFrame(MoveTemp(Frame));
+	}
+}
+
+void FFrameGrabberProtocol::Finalize()
+{
+	FrameGrabber->Shutdown();
+	FrameGrabber.Reset();
+}
\ No newline at end of file
diff --git a/Engine/Source/Runtime/MovieSceneCapture/Private/ImageSequenceProtocol.cpp b/Engine/Source/Runtime/MovieSceneCapture/Private/ImageSequenceProtocol.cpp
new file mode 100644
--- /dev/null
+++ b/Engine/Source/Runtime/MovieSceneCapture/Private/ImageSequenceProtocol.cpp
@@ -0,0 +1,259 @@
+// Copyright 1998-2015 Epic Games, Inc. All Rights Reserved.
+
+#include "MovieSceneCapturePCH.h"
+
+#include "ImageSequenceProtocol.h"
+
+#if WITH_EDITOR
+#include "IImageWrapperModule.h"
+
+namespace
+{
+	static const int32 MaxAsyncWrites = 6;
+}
+
+struct FImageFrameData : IFramePayload
+{
+	FString Filename;
+};
+
+FImageSequenceProtocol::FImageSequenceProtocol(EImageFormat::Type InFormat)
+{
+	Format = InFormat;
+	CompressionQuality = 100;
+}
+
+bool FImageSequenceProtocol::Initialize(const FCaptureProtocolInitSettings& InSettings, const ICaptureProtocolHost& Host)
+{
+	if (!FFrameGrabberProtocol::Initialize(InSettings, Host))
+	{
+		return false;
+	}
+
+	UImageCaptureSettings* CaptureSettings = Cast<UImageCaptureSettings>(InSettings.ProtocolSettings);
+	if (CaptureSettings)
+	{
+		CompressionQuality = CaptureSettings->CompressionQuality;
+		
+		FParse::Value( FCommandLine::Get(), TEXT( "-MovieQuality=" ), CompressionQuality );
+
+		CompressionQuality = FMath::Clamp<int32>(CompressionQuality, 1, 100);
+	}
+
+	CaptureThread.Reset(new FImageCaptureThread(Format, CompressionQuality));
+	
+	return true;
+}
+
+bool FImageSequenceProtocol::HasFinishedProcessing() const
+{
+	return FFrameGrabberProtocol::HasFinishedProcessing() && CaptureThread->GetNumOutstandingFrames() == 0;
+}
+
+void FImageSequenceProtocol::Finalize()
+{
+	CaptureThread->Close();
+	CaptureThread.Reset();
+
+	FFrameGrabberProtocol::Finalize();
+}
+
+FFramePayloadPtr FImageSequenceProtocol::GetFramePayload(const FFrameMetrics& FrameMetrics, const ICaptureProtocolHost& Host) const
+{
+	TSharedRef<FImageFrameData, ESPMode::ThreadSafe> FrameData = MakeShareable(new FImageFrameData);
+
+	const TCHAR* Extension = TEXT("");
+	switch(Format)
+	{
+	case EImageFormat::BMP:		Extension = TEXT(".bmp"); break;
+	case EImageFormat::PNG:		Extension = TEXT(".png"); break;
+	case EImageFormat::JPEG:	Extension = TEXT(".jpg"); break;
+	}
+
+	FrameData->Filename = Host.GenerateFilename(FrameMetrics, Extension);
+
+	// Add our custom formatting rules as well
+	// @todo: document these on the tooltip?
+	FrameData->Filename = FString::Format(*FrameData->Filename, StringFormatMap);
+
+	return FrameData;
+}
+
+void FImageSequenceProtocol::ProcessFrame(FCapturedFrameData Frame)
+{
+	CaptureThread->Add(MoveTemp(Frame));
+}
+
+void FImageSequenceProtocol::AddFormatMappings(TMap<FString, FStringFormatArg>& FormatMappings) const
+{
+	if (Format == EImageFormat::JPEG || Format == EImageFormat::PNG)
+	{
+		FormatMappings.Add(TEXT("quality"), CompressionQuality);
+	}
+	else
+	{
+		FormatMappings.Add(TEXT("quality"), TEXT(""));
+	}
+}
+
+
+FImageCaptureThread::FImageCaptureThread(EImageFormat::Type InFormat, int32 InCompressionQuality)
+{
+	Format = InFormat;
+	CompressionQuality = InCompressionQuality;
+
+	WorkToDoEvent = FPlatformProcess::GetSynchEventFromPool();
+	ThreadEmptyEvent = FPlatformProcess::GetSynchEventFromPool();
+
+	CapturedFrames.Reserve(MaxAsyncWrites);
+
+	if (Format == EImageFormat::PNG || Format == EImageFormat::JPEG)
+	{
+		IImageWrapperModule& ImageWrapperModule = FModuleManager::LoadModuleChecked<IImageWrapperModule>( FName("ImageWrapper") );
+
+		for (int32 Index = 0; Index < MaxAsyncWrites; ++Index)
+		{
+			ImageWrappers.Add(ImageWrapperModule.CreateImageWrapper(Format));
+		}
+	}
+
+	static int32 Index = 0;
+	Thread = FRunnableThread::Create(this, *FString::Printf(TEXT("ImageCaptureThread_%d"), ++Index));
+}
+
+FImageCaptureThread::~FImageCaptureThread()
+{
+	FPlatformProcess::ReturnSynchEventToPool(WorkToDoEvent);
+	FPlatformProcess::ReturnSynchEventToPool(ThreadEmptyEvent);
+}
+
+void FImageCaptureThread::Add(FCapturedFrameData Frame)
+{
+	bool bThreadIsChoked = false;
+	{
+		FScopeLock Lock(&CommandMutex);
+		CapturedFrames.Add(MoveTemp(Frame));
+
+		bThreadIsChoked = CapturedFrames.Num() > MaxAsyncWrites;
+	}
+
+	WorkToDoEvent->Trigger();
+
+	if (bThreadIsChoked)
+	{
+		ThreadEmptyEvent->Wait(~0);
+	}
+}
+
+uint32 FImageCaptureThread::GetNumOutstandingFrames() const
+{
+	FScopeLock Lock(&CommandMutex);
+	return CapturedFrames.Num();
+}
+
+void FImageCaptureThread::Close()
+{
+	Thread->Kill(true);
+}
+
+void FImageCaptureThread::Stop()
+{
+	bRunning = false;
+	WorkToDoEvent->Trigger();
+}
+
+uint32 FImageCaptureThread::Run()
+{
+	bRunning = true;
+
+	TArray<TFuture<void>> FrameFutures;
+	TArray<FCapturedFrameData> ProcessingFrames;
+
+	ProcessingFrames.Reserve(MaxAsyncWrites);
+	FrameFutures.Reserve(MaxAsyncWrites);
+
+	for(;;)
+	{
+		WorkToDoEvent->Wait(~0);
+		if (!bRunning)
+		{
+			ThreadEmptyEvent->Trigger();
+			return 0;
+		}
+
+		// process all outstanding frames, MaxAsyncWrites at a time
+		for (;;)
+		{
+			{
+				FScopeLock Lock(&CommandMutex);
+				for (int32 Index = 0; Index < MaxAsyncWrites && Index < CapturedFrames.Num(); ++Index)
+				{
+					ProcessingFrames.Add(MoveTemp(CapturedFrames[Index]));
+				}
+
+				if (ProcessingFrames.Num())
+				{
+					CapturedFrames.RemoveAt(0, ProcessingFrames.Num(), false);
+				}
+				else
+				{
+					break;
+				}
+			}
+
+			for (int32 FrameIndex = 0; FrameIndex < ProcessingFrames.Num(); ++FrameIndex)
+			{
+				FrameFutures.Add(Async<void>(EAsyncExecution::TaskGraph, [&, FrameIndex]{
+					IImageWrapperPtr ImageWrapper = ImageWrappers.Num() ? ImageWrappers[FrameIndex] : nullptr;
+					WriteFrameToDisk(ProcessingFrames[FrameIndex], ImageWrapper);
+				}));
+			}
+
+			for (TFuture<void>& Future : FrameFutures)
+			{
+				Future.Wait();
+			}
+
+			ProcessingFrames.Reset();
+			FrameFutures.Reset();
+		}
+
+		ThreadEmptyEvent->Trigger();
+	}
+
+	return 0;
+}
+
+void FImageCaptureThread::WriteFrameToDisk(FCapturedFrameData& Frame, IImageWrapperPtr ImageWrapper) const
+{
+	const float Width = Frame.BufferSize.X;
+	const float Height = Frame.BufferSize.Y;
+
+	FImageFrameData* Payload = Frame.GetPayload<FImageFrameData>();
+	switch (Format)
+	{
+	case EImageFormat::BMP:
+		FFileHelper::CreateBitmap(*Payload->Filename, Width, Height, Frame.ColorBuffer.GetData());
+		break;
+	case EImageFormat::PNG:
+		for (FColor& Color : Frame.ColorBuffer)
+		{
+			Color.A = 255;
+		}
+
+		if (ImageWrapper->SetRaw(Frame.ColorBuffer.GetData(), Frame.ColorBuffer.Num() * sizeof(FColor), Width, Height, ERGBFormat::BGRA, 8))
+		{
+			FFileHelper::SaveArrayToFile(ImageWrapper->GetCompressed(CompressionQuality), *Payload->Filename);
+		}
+		break;
+
+	case EImageFormat::JPEG:
+		if (ImageWrapper->SetRaw(Frame.ColorBuffer.GetData(), Frame.ColorBuffer.Num() * sizeof(FColor), Width, Height, ERGBFormat::BGRA, 8))
+		{
+			FFileHelper::SaveArrayToFile(ImageWrapper->GetCompressed(CompressionQuality), *Payload->Filename);
+		}
+		break;
+	}
+}
+
+#endif
\ No newline at end of file
diff --git a/Engine/Source/Runtime/MovieSceneCapture/Private/LevelCapture.cpp b/Engine/Source/Runtime/MovieSceneCapture/Private/LevelCapture.cpp
--- a/Engine/Source/Runtime/MovieSceneCapture/Private/LevelCapture.cpp
+++ b/Engine/Source/Runtime/MovieSceneCapture/Private/LevelCapture.cpp
@@ -3,45 +3,56 @@
 #include "MovieSceneCapturePCH.h"
 #include "LevelCapture.h"
 
-void ULevelCapture::SetPrerequisiteActor(AActor* InPrerequisiteActor)
+ULevelCapture::ULevelCapture(const FObjectInitializer& Init)
+	: Super(Init)
 {
-	PrerequisiteActor = InPrerequisiteActor;
-	PrerequisiteActorId = PrerequisiteActor.GetUniqueID().GetGuid();
+	bAutoStartCapture = true;
 }
 
-void ULevelCapture::Initialize(TWeakPtr<FSceneViewport> InViewport)
+void ULevelCapture::SetPrerequisiteActor(AActor* InPrerequisiteActor)
 {
-	// Ensure the PrerequisiteActor is up to date with the PrerequisiteActorId, since this may have been deserialized from JSON
-	PrerequisiteActor = FUniqueObjectGuid(PrerequisiteActorId);
+	PrerequisiteActorId = TLazyObjectPtr<>(InPrerequisiteActor).GetUniqueID().GetGuid();
+}
 
-	Super::Initialize(InViewport);
+void ULevelCapture::Initialize(TSharedPtr<FSceneViewport> InViewport, int32 PIEInstance)
+{
 	CaptureStrategy = MakeShareable(new FFixedTimeStepCaptureStrategy(Settings.FrameRate));
+	Super::Initialize(InViewport, PIEInstance);
+
+	PIECaptureInstance = PIEInstance;
+
+	if (bAutoStartCapture)
+	{
+		StartCapture();
+	}
 }
 
 void ULevelCapture::Tick(float DeltaSeconds)
 {
-	Super::Tick(DeltaSeconds);
-
 	if (!GWorld->HasBegunPlay())
 	{
 		return;
 	}
 
-	AActor* PrerequisiteActorPtr = PrerequisiteActor.Get();
-	if (!PrerequisiteActorId.IsValid() || (PrerequisiteActorPtr && PrerequisiteActorPtr->HasActorBegunPlay()))
+	AActor* Actor = PrerequisiteActor.Get();
+	if (!Actor)
 	{
-		if (!bCapturing)
+		// We need to look the actor up
+		TLazyObjectPtr<AActor> LazyActor;
+		if (PIECaptureInstance != -1)
+		{
+			LazyActor = FUniqueObjectGuid(PrerequisiteActorId).FixupForPIE(PIECaptureInstance);
+		}
+		else
 		{
-			StartCapture();
+			LazyActor = FUniqueObjectGuid(PrerequisiteActorId);
 		}
 
-		PrepareForScreenshot();
+		PrerequisiteActor = Actor = LazyActor.Get();
 	}
-}
 
-#if WITH_EDITOR
-void ULevelCapture::PostEditChangeProperty(struct FPropertyChangedEvent& PropertyChangedEvent)
-{
-	PrerequisiteActorId = PrerequisiteActor.GetUniqueID().GetGuid();
-}
-#endif
\ No newline at end of file
+	if (!PrerequisiteActorId.IsValid() || (Actor && Actor->HasActorBegunPlay()))
+	{
+		CaptureThisFrame(DeltaSeconds);
+	}
+}
\ No newline at end of file
diff --git a/Engine/Source/Runtime/MovieSceneCapture/Private/MovieSceneCapture.cpp b/Engine/Source/Runtime/MovieSceneCapture/Private/MovieSceneCapture.cpp
--- a/Engine/Source/Runtime/MovieSceneCapture/Private/MovieSceneCapture.cpp
+++ b/Engine/Source/Runtime/MovieSceneCapture/Private/MovieSceneCapture.cpp
@@ -3,6 +3,9 @@
 #include "MovieSceneCapturePCH.h"
 #include "MovieSceneCapture.h"
 #include "ActiveMovieSceneCaptures.h"
+#include "HighResScreenshot.h"
+#include "BufferVisualizationData.h"
+#include "SceneViewExtension.h"
 
 #if WITH_EDITOR
 #include "ImageWrapper.h"
@@ -31,11 +34,8 @@ FMovieSceneCaptureSettings::FMovieSceneCaptureSettings()
 	bCreateTemporaryCopiesOfLevels = false;
 	bUseRelativeFrameNumbers = false;
 	GameModeOverride = nullptr;
-	OutputFormat = NSLOCTEXT("MovieCapture", "DefaultFormat", "{world}_{frame}").ToString();
+	OutputFormat = TEXT("{world}_{frame}");
 	FrameRate = 24;
-	CaptureType = EMovieCaptureType::AVI;
-	bUseCompression = true;
-	CompressionQuality = 1.f;
 	bEnableTextureStreaming = true;
 	bCinematicMode = true;
 	bAllowMovement = false;
@@ -47,8 +47,6 @@ FMovieSceneCaptureSettings::FMovieSceneCaptureSettings()
 UMovieSceneCapture::UMovieSceneCapture(const FObjectInitializer& Initializer)
 	: Super(Initializer)
 {
-	bBufferVisualizationDumpFrames = false;
-
 	TArray<FString> Tokens, Switches;
 	FCommandLine::Parse( FCommandLine::Get(), Tokens, Switches );
 	for (auto& Switch : Switches)
@@ -63,14 +61,29 @@ UMovieSceneCapture::UMovieSceneCapture(const FObjectInitializer& Initializer)
 	Handle = FUniqueMovieSceneCaptureHandle();
 
 	FrameCount = 0;
-	OutstandingFrameCount = 0;
-	LastFrameDelta = 0.f;
 	bCapturing = false;
 	FrameNumberOffset = 0;
+	CaptureType = TEXT("Video");
+}
+
+void UMovieSceneCapture::PostInitProperties()
+{
+	if (!HasAnyFlags(RF_ClassDefaultObject))
+	{
+		ProtocolSettings = IMovieSceneCaptureModule::Get().GetProtocolRegistry().FactorySettingsType(CaptureType, this);
+		if (ProtocolSettings)
+		{
+			ProtocolSettings->LoadConfig();
+		}
+	}
+
+	Super::PostInitProperties();
 }
 
-void UMovieSceneCapture::Initialize(TWeakPtr<FSceneViewport> InSceneViewport)
+void UMovieSceneCapture::Initialize(TSharedPtr<FSceneViewport> InSceneViewport, int32 PIEInstance)
 {
+	ensure(!bCapturing);
+
 	// Apply command-line overrides
 	{
 		FString OutputPathOverride;
@@ -106,22 +119,7 @@ void UMovieSceneCapture::Initialize(TWeakPtr<FSceneViewport> InSceneViewport)
 		FString FormatOverride;
 		if( FParse::Value( FCommandLine::Get(), TEXT( "-MovieFormat=" ), FormatOverride ) )
 		{
-			if( FormatOverride.Equals( TEXT( "AVI" ), ESearchCase::IgnoreCase ) )
-			{
-				Settings.CaptureType = EMovieCaptureType::AVI;
-			}
-			else if( FormatOverride.Equals( TEXT( "BMP" ), ESearchCase::IgnoreCase ) )
-			{
-				Settings.CaptureType = EMovieCaptureType::BMP;
-			}
-			else if( FormatOverride.Equals( TEXT( "PNG" ), ESearchCase::IgnoreCase ) )
-			{
-				Settings.CaptureType = EMovieCaptureType::PNG;
-			}
-			else if( FormatOverride.Equals( TEXT( "JPG" ), ESearchCase::IgnoreCase ) || FormatOverride.Equals( TEXT( "JPEG" ), ESearchCase::IgnoreCase ) )
-			{
-				Settings.CaptureType = EMovieCaptureType::JPEG;
-			}
+			CaptureType = *FormatOverride;
 		}
 
 		int32 FrameRateOverride;
@@ -129,74 +127,26 @@ void UMovieSceneCapture::Initialize(TWeakPtr<FSceneViewport> InSceneViewport)
 		{
 			Settings.FrameRate = FrameRateOverride;
 		}
-
-		int32 QualityOverride;
-		if( FParse::Value( FCommandLine::Get(), TEXT( "-MovieQuality=" ), QualityOverride ) )
-		{
-			Settings.CompressionQuality = FMath::TruncToFloat( FMath::Clamp( QualityOverride, 0, 100 ) ) / 100.0f;
-			Settings.bUseCompression = QualityOverride <= 0.99f;
-		}
 	}
 
-	SceneViewport = InSceneViewport;
+	bFinalizeWhenReady = false;
 
-	if (!HasAnyFlags(RF_ClassDefaultObject))
-	{
-		Ticker.Reset(new FTicker(this));
-	}
-}
-
-void UMovieSceneCapture::PrepareForScreenshot()
-{
-	auto Viewport = SceneViewport.Pin();
-	if (!Viewport.IsValid())
-	{
-		return;
-	}
+	InitSettings = FCaptureProtocolInitSettings::FromSlateViewport(InSceneViewport.ToSharedRef(), ProtocolSettings);
 
-	auto ViewportWidget = Viewport->GetViewportWidget().Pin();
-	if (!ViewportWidget.IsValid())
-	{
-		return;
-	}
+	CachedMetrics = FCachedMetrics();
+	CachedMetrics.Width = InitSettings->DesiredSize.X;
+	CachedMetrics.Height = InitSettings->DesiredSize.Y;
 
-	auto& Application = FSlateApplication::Get();
+	FormatMappings.Reserve(10);
+	FormatMappings.Add(TEXT("fps"), FString::Printf(TEXT("%d"), Settings.FrameRate));
+	FormatMappings.Add(TEXT("width"), FString::Printf(TEXT("%d"), CachedMetrics.Width));
+	FormatMappings.Add(TEXT("height"), FString::Printf(TEXT("%d"), CachedMetrics.Height));
+	FormatMappings.Add(TEXT("world"), GWorld->GetName());
 
-	// We can't screenshot the widget unless there's a valid window handle to draw it in.
-	TSharedPtr<SWindow> WidgetWindow = Application.FindWidgetWindow(ViewportWidget.ToSharedRef());
-	if (!WidgetWindow.IsValid())
+	if (!HasAnyFlags(RF_ClassDefaultObject))
 	{
-		return;
+		FActiveMovieSceneCaptures::Get().Add(this);
 	}
-
-	++OutstandingFrameCount;
-
-	FWidgetPath WidgetPath;
-	FSlateApplication::Get().GeneratePathToWidgetChecked(ViewportWidget.ToSharedRef(), WidgetPath);
-
-	FArrangedWidget ArrangedWidget = WidgetPath.FindArrangedWidget(ViewportWidget.ToSharedRef()).Get(FArrangedWidget::NullWidget);
-
-	FVector2D Position = ArrangedWidget.Geometry.AbsolutePosition;
-	FVector2D WindowPosition = WidgetWindow->GetPositionInScreen();
-
-	const int32 RelativePositionX = Position.X - WindowPosition.X;
-	const int32 RelativePositionY = Position.Y - WindowPosition.Y;
-
-	FVector2D Size = ArrangedWidget.Geometry.GetDrawSize();
-
-	FIntRect ScreenshotRect(
-		RelativePositionX,
-		RelativePositionY,
-		RelativePositionX + Size.X,
-		RelativePositionY + Size.Y);
-
-	Application.GetRenderer()->PrepareToTakeScreenshot(ScreenshotRect, &ScratchBuffer);
-}
-
-void UMovieSceneCapture::Tick(float DeltaSeconds)
-{
-	CaptureFrame(LastFrameDelta);
-	LastFrameDelta = DeltaSeconds;
 }
 
 void UMovieSceneCapture::StartWarmup()
@@ -211,12 +161,7 @@ void UMovieSceneCapture::StartWarmup()
 
 void UMovieSceneCapture::StartCapture()
 {
-	auto Viewport = SceneViewport.Pin();
-	if (!Viewport.IsValid())
-	{
-		return;
-	}
-
+	bFinalizeWhenReady = false;
 	bCapturing = true;
 
 	if (!CaptureStrategy.IsValid())
@@ -224,214 +169,99 @@ void UMovieSceneCapture::StartCapture()
 		CaptureStrategy = MakeShareable(new FRealTimeCaptureStrategy(Settings.FrameRate));
 	}
 
-	check( CaptureStrategy.IsValid() );
 	CaptureStrategy->OnStart();
 
-	CachedMetrics.ElapsedSeconds = 0;
-
-	if (bBufferVisualizationDumpFrames)
+	CaptureProtocol = IMovieSceneCaptureModule::Get().GetProtocolRegistry().Factory(CaptureType);
+	if (ensure(CaptureProtocol.IsValid()))
 	{
-		static IConsoleVariable* CVarDumpFrames = IConsoleManager::Get().FindConsoleVariable(TEXT("r.BufferVisualizationDumpFrames"));
-		if (CVarDumpFrames)
-		{
-			CVarDumpFrames->Set(1, ECVF_SetByCommandline);
-		}
-	}
-
-	CachedMetrics.Width = Viewport->GetSize().X;
-	CachedMetrics.Height = Viewport->GetSize().Y;
-
-	if (Settings.CaptureType == EMovieCaptureType::AVI)
-	{
-		FAVIWriterOptions Options;
-		Options.OutputFilename = ResolveUniqueFilename();
-		Options.CaptureFPS = Settings.FrameRate;
-		Options.CodecName = Settings.Codec;
-		Options.bSynchronizeFrames = CaptureStrategy->ShouldSynchronizeFrames();
-
-		if (Settings.bUseCompression)
-		{
-			Options.CompressionQuality = Settings.CompressionQuality;
-		}
-
-		AVIWriter.Reset(FAVIWriter::CreateInstance(Options));
-		AVIWriter->StartCapture(Viewport);
+		CaptureProtocol->Initialize(InitSettings.GetValue(), *this);
 	}
 }
 
-void UMovieSceneCapture::CaptureFrame(float DeltaSeconds)
+void UMovieSceneCapture::CaptureThisFrame(float DeltaSeconds)
 {
-	auto Viewport = SceneViewport.Pin();
-	if (OutstandingFrameCount == 0 || !CaptureStrategy.IsValid() || !Viewport.IsValid() || !ScratchBuffer.Num())
+	if (!bCapturing || !CaptureStrategy.IsValid() || !CaptureProtocol.IsValid() || bFinalizeWhenReady)
 	{
 		return;
 	}
-	
-	--OutstandingFrameCount;
 
 	CachedMetrics.ElapsedSeconds += DeltaSeconds;
-
-	// By this point, the slate application has already populated our scratch buffer
 	if (CaptureStrategy->ShouldPresent(CachedMetrics.ElapsedSeconds, CachedMetrics.Frame))
 	{
-		TArray<FColor> ThisFrameBuffer;
-		Swap(ThisFrameBuffer, ScratchBuffer);
-		ScratchBuffer.Reset();
-
 		uint32 NumDroppedFrames = CaptureStrategy->GetDroppedFrames(CachedMetrics.ElapsedSeconds, CachedMetrics.Frame);
 		CachedMetrics.Frame += NumDroppedFrames;
 
 		CaptureStrategy->OnPresent(CachedMetrics.ElapsedSeconds, CachedMetrics.Frame);
-		
-		if (AVIWriter)
-		{
-			AVIWriter->DropFrames(NumDroppedFrames);
-			AVIWriter->Update(CachedMetrics.ElapsedSeconds, MoveTemp(ThisFrameBuffer));
-		}
-#if WITH_EDITOR
-		else
-		{
-			SaveFrameToFile(MoveTemp(ThisFrameBuffer));
-		}
-#endif
-		++CachedMetrics.Frame;
-		if (FrameCount != 0 && CachedMetrics.Frame >= FrameCount)
-		{
-			StopCapture();
-		}
-	}
-}
-
-void UMovieSceneCapture::StopCapture()
-{
-	Ticker.Reset();
 
-	if (bCapturing)
-	{
-		OutstandingFrameCount = 0;
-		bCapturing = false;
+		const FFrameMetrics ThisFrameMetrics(
+			CachedMetrics.ElapsedSeconds,
+			DeltaSeconds,
+			CachedMetrics.Frame,
+			NumDroppedFrames
+			);
+		CaptureProtocol->CaptureFrame(ThisFrameMetrics, *this);
 
-		CaptureStrategy->OnStop();
-		CaptureStrategy = nullptr;
+		++CachedMetrics.Frame;
 
-		if (AVIWriter)
+		if (!bFinalizeWhenReady && (FrameCount != 0 && CachedMetrics.Frame >= FrameCount))
 		{
-			AVIWriter->StopCapture();
-			AVIWriter.Reset();
+			FinalizeWhenReady();
 		}
-
-		OnCaptureStopped();
 	}
 }
 
-void UMovieSceneCapture::Close()
+void UMovieSceneCapture::FinalizeWhenReady()
 {
-	StopCapture();
-	FActiveMovieSceneCaptures::Get().Remove(this);
+	bFinalizeWhenReady = true;
 }
 
-#if WITH_EDITOR
-void UMovieSceneCapture::SaveFrameToFile(TArray<FColor> Colors)
+void UMovieSceneCapture::Finalize()
 {
-	if (Colors.Num() == 0)
-	{
-		return;
-	}
-
-	FString Filename = ResolveUniqueFilename();
-	if (Filename.IsEmpty())
-	{
-		return;
-	}
-
-	IImageWrapperModule& ImageWrapperModule = FModuleManager::LoadModuleChecked<IImageWrapperModule>( FName("ImageWrapper") );
-
-	// clamp to 1-100 range
-	const int32 ImageCompressionQuality = FMath::Clamp<int32>(Settings.CompressionQuality * 100, 1, 100);
+	FActiveMovieSceneCaptures::Get().Remove(this);
 
-	switch (Settings.CaptureType)
+	if (bCapturing)
 	{
-	case EMovieCaptureType::BMP:
-		FFileHelper::CreateBitmap(*Filename, CachedMetrics.Width, CachedMetrics.Height, Colors.GetData());
-		break;
+		bCapturing = false;
 
-	case EMovieCaptureType::PNG:
+		if (CaptureStrategy.IsValid())
 		{
-			IImageWrapperPtr ImageWrapper = ImageWrapperModule.CreateImageWrapper(EImageFormat::PNG);
-			for (FColor& Color : Colors)
-			{
-				Color.A = 255;
-			}
-
-			if (ImageWrapper.IsValid() && ImageWrapper->SetRaw(&Colors[0], Colors.Num() * sizeof(FColor), CachedMetrics.Width, CachedMetrics.Height, ERGBFormat::BGRA, 8))
-			{
-				FFileHelper::SaveArrayToFile(ImageWrapper->GetCompressed(ImageCompressionQuality), *Filename);
-			}
+			CaptureStrategy->OnStop();
+			CaptureStrategy = nullptr;
 		}
-		break;
 
-	case EMovieCaptureType::JPEG:
+		if (CaptureProtocol.IsValid())
 		{
-			IImageWrapperPtr ImageWrapper = ImageWrapperModule.CreateImageWrapper(EImageFormat::JPEG);
-			if (ImageWrapper.IsValid() && ImageWrapper->SetRaw(&Colors[0], Colors.Num() * sizeof(FColor), CachedMetrics.Width, CachedMetrics.Height, ERGBFormat::BGRA, 8))
-			{
-				FFileHelper::SaveArrayToFile(ImageWrapper->GetCompressed(ImageCompressionQuality), *Filename);
-			}
+			CaptureProtocol->Finalize();
+			CaptureProtocol = nullptr;
 		}
-		break;
 
-	default:
-		break;
+		OnCaptureFinishedDelegate.Broadcast();
 	}
 }
-#endif
 
-const TCHAR* UMovieSceneCapture::GetDefaultFileExtension() const
+FString UMovieSceneCapture::ResolveFileFormat(const FString& Format, const FFrameMetrics& FrameMetrics) const
 {
-	switch (Settings.CaptureType)
+	FormatMappings.Add(TEXT("frame"), FString::Printf(TEXT("%04d"), Settings.bUseRelativeFrameNumbers ? FrameMetrics.FrameNumber : FrameMetrics.FrameNumber + FrameNumberOffset));
+	
+	if (CaptureProtocol.IsValid())
 	{
-		case EMovieCaptureType::BMP:	return TEXT(".bmp");
-		case EMovieCaptureType::PNG:	return TEXT(".png");
-		case EMovieCaptureType::JPEG:	return TEXT(".jpg");
-#if PLATFORM_MAC
-		default:						return TEXT(".mov");
-#else
-		default:						return TEXT(".avi");
-#endif
+		CaptureProtocol->AddFormatMappings(FormatMappings);
 	}
+	return FString::Format(*Format, FormatMappings);
 }
 
-FString UMovieSceneCapture::ResolveFileFormat(const FString& Folder, const FString& Format) const
+FString UMovieSceneCapture::GenerateFilename(const FFrameMetrics& FrameMetrics, const TCHAR* Extension) const
 {
-	TMap<FString, FStringFormatArg> Mappings;
-	Mappings.Add(TEXT("fps"), FString::Printf(TEXT("%d"), Settings.FrameRate));
-	Mappings.Add(TEXT("frame"), FString::Printf(TEXT("%04d"), Settings.bUseRelativeFrameNumbers ? CachedMetrics.Frame : CachedMetrics.Frame + FrameNumberOffset));
-	Mappings.Add(TEXT("width"), FString::Printf(TEXT("%d"), CachedMetrics.Width));
-	Mappings.Add(TEXT("height"), FString::Printf(TEXT("%d"), CachedMetrics.Height));
-	Mappings.Add(TEXT("world"), GWorld->GetName());
-
-	if (Settings.bUseCompression)
-	{
-		Mappings.Add(TEXT("quality"), FString::Printf(TEXT("%.2f"), Settings.CompressionQuality));
-	}
-	else
-	{
-		Mappings.Add(TEXT("quality"), NSLOCTEXT("MovieCapture", "Uncompressed", "Uncompressed").ToString());
-	}
+	FString OutputDirectory = ResolveFileFormat(Settings.OutputDirectory.Path, FrameMetrics);
 
-	return Folder / FString::Format(*Format, Mappings);
-}
-
-FString UMovieSceneCapture::ResolveUniqueFilename()
-{
-	if( !IFileManager::Get().DirectoryExists( *Settings.OutputDirectory.Path ) )
+	if (!IFileManager::Get().DirectoryExists(*OutputDirectory))
 	{
-		IFileManager::Get().MakeDirectory( *Settings.OutputDirectory.Path );
+		IFileManager::Get().MakeDirectory(*OutputDirectory);
 	}
 
-	const FString BaseFilename = ResolveFileFormat( Settings.OutputDirectory.Path, Settings.OutputFormat );
+	const FString BaseFilename = OutputDirectory / ResolveFileFormat(Settings.OutputFormat, FrameMetrics);
 
-	FString ThisTry = BaseFilename + GetDefaultFileExtension();
+	FString ThisTry = BaseFilename + Extension;
 
 	if (Settings.bOverwriteExisting)
 	{
@@ -453,7 +283,7 @@ FString UMovieSceneCapture::ResolveUniqueFilename()
 	int32 DuplicateIndex = 1;
 	for (;;)
 	{
-		ThisTry = BaseFilename + FString::Printf(TEXT("_(%d)"), DuplicateIndex) + GetDefaultFileExtension();
+		ThisTry = BaseFilename + FString::Printf(TEXT("_(%d)"), DuplicateIndex) + Extension;
 
 		// If the file doesn't exist, we can use that, else, increment the index and try again
 		if (IFileManager::Get().FileSize(*ThisTry) == -1)
@@ -467,6 +297,24 @@ FString UMovieSceneCapture::ResolveUniqueFilename()
 	return ThisTry;
 }
 
+#if WITH_EDITOR
+void UMovieSceneCapture::PostEditChangeProperty( struct FPropertyChangedEvent& PropertyChangedEvent)
+{
+	FName PropertyName = (PropertyChangedEvent.MemberProperty != NULL) ? PropertyChangedEvent.MemberProperty->GetFName() : NAME_None;
+
+	if (PropertyName == GET_MEMBER_NAME_CHECKED(UMovieSceneCapture, CaptureType))
+	{
+		ProtocolSettings = IMovieSceneCaptureModule::Get().GetProtocolRegistry().FactorySettingsType(CaptureType, this);
+		if (ProtocolSettings)
+		{
+			ProtocolSettings->LoadConfig();
+		}
+	}
+
+	Super::PostEditChangeProperty(PropertyChangedEvent);
+}
+#endif
+
 FFixedTimeStepCaptureStrategy::FFixedTimeStepCaptureStrategy(uint32 InTargetFPS)
 	: TargetFPS(InTargetFPS)
 {
diff --git a/Engine/Source/Runtime/MovieSceneCapture/Private/MovieSceneCaptureModule.cpp b/Engine/Source/Runtime/MovieSceneCapture/Private/MovieSceneCaptureModule.cpp
--- a/Engine/Source/Runtime/MovieSceneCapture/Private/MovieSceneCaptureModule.cpp
+++ b/Engine/Source/Runtime/MovieSceneCapture/Private/MovieSceneCaptureModule.cpp
@@ -6,6 +6,10 @@
 #include "JsonObjectConverter.h"
 #include "ActiveMovieSceneCaptures.h"
 
+#include "Protocols/ImageSequenceProtocol.h"
+#include "Protocols/CompositionGraphCaptureProtocol.h"
+#include "Protocols/VideoCaptureProtocol.h"
+
 #define LOCTEXT_NAMESPACE "MovieSceneCapture"
 
 class FMovieSceneCaptureModule : public IMovieSceneCaptureModule
@@ -14,12 +18,62 @@ private:
 
 	/** Handle to a movie capture implementation created from the command line, to be initialized once a world is loaded */
 	FMovieSceneCaptureHandle StartupMovieCaptureHandle;
+	FMovieSceneCaptureProtocolRegistry ProtocolRegistry;
+
+	virtual FMovieSceneCaptureProtocolRegistry& GetProtocolRegistry()
+	{
+		return ProtocolRegistry;
+	}
 
-	/** End Tickable interface */
+	
 	virtual void StartupModule() override
 	{
 		FCoreDelegates::OnPreExit.AddRaw(this, &FMovieSceneCaptureModule::PreExit);
 		FCoreUObjectDelegates::PostLoadMap.AddRaw(this, &FMovieSceneCaptureModule::OnPostLoadMap );
+
+		FMovieSceneCaptureProtocolInfo Info;
+		{
+			Info.DisplayName = LOCTEXT("CompositionGraphDescription", "Custom Render Passes");
+			Info.SettingsClassType = UCompositionGraphCaptureSettings::StaticClass();
+			Info.Factory = []() -> TSharedRef<IMovieSceneCaptureProtocol> {
+				return MakeShareable(new FCompositionGraphCaptureProtocol());
+			};
+			ProtocolRegistry.RegisterProtocol(TEXT("CustomRenderPasses"), Info);
+		}
+#if WITH_EDITOR
+		{
+			Info.DisplayName = LOCTEXT("VideoDescription", "Video Sequence");
+			Info.SettingsClassType = UVideoCaptureSettings::StaticClass();
+			Info.Factory = []() -> TSharedRef<IMovieSceneCaptureProtocol> {
+				return MakeShareable(new FVideoCaptureProtocol);
+			};
+			ProtocolRegistry.RegisterProtocol(TEXT("Video"), Info);
+		}
+		{
+			Info.DisplayName = LOCTEXT("PNGDescription", "Image Sequence (png)");
+			Info.SettingsClassType = UImageCaptureSettings::StaticClass();
+			Info.Factory = []() -> TSharedRef<IMovieSceneCaptureProtocol> {
+				return MakeShareable(new FImageSequenceProtocol(EImageFormat::PNG));
+			};
+			ProtocolRegistry.RegisterProtocol(TEXT("PNG"), Info);
+		}
+		{
+			Info.DisplayName = LOCTEXT("JPEGDescription", "Image Sequence (jpg)");
+			Info.SettingsClassType = UImageCaptureSettings::StaticClass();
+			Info.Factory = []() -> TSharedRef<IMovieSceneCaptureProtocol> {
+				return MakeShareable(new FImageSequenceProtocol(EImageFormat::JPEG));
+			};
+			ProtocolRegistry.RegisterProtocol(TEXT("JPG"), Info);
+		}
+		{
+			Info.DisplayName = LOCTEXT("BMPDescription", "Image Sequence (bmp)");
+			Info.SettingsClassType = nullptr;	// Bitmaps don't have any options
+			Info.Factory = []() -> TSharedRef<IMovieSceneCaptureProtocol> {
+				return MakeShareable(new FImageSequenceProtocol(EImageFormat::BMP));
+			};
+			ProtocolRegistry.RegisterProtocol(TEXT("BMP"), Info);
+		}
+#endif
 	}
 
 	void PreExit()
@@ -34,7 +88,6 @@ private:
 		{
 			IMovieSceneCaptureInterface* StartupCaptureInterface = RetrieveMovieSceneInterface(StartupMovieCaptureHandle);
 			StartupCaptureInterface->Initialize(GameEngine->SceneViewport.ToSharedRef());
-			StartupCaptureInterface->SetupFrameRange();
 		}
 
 		StartupMovieCaptureHandle = FMovieSceneCaptureHandle();
@@ -43,8 +96,7 @@ private:
 
 	virtual void PreUnloadCallback() override
 	{
-		FCoreDelegates::OnPreExit.RemoveAll(this);
-		PreExit();
+		DestroyAllActiveCaptures();
 	}
 
 	virtual IMovieSceneCaptureInterface* InitializeFromCommandLine() override
@@ -109,6 +161,25 @@ private:
 					{
 						return nullptr;
 					}
+
+					// Now deserialize the protocol data
+					auto ProtocolTypeField = RootObject->TryGetField(TEXT("ProtocolType"));
+					if (ProtocolTypeField.IsValid())
+					{
+						UClass* ProtocolTypeClass = FindObject<UClass>(nullptr, *ProtocolTypeField->AsString());
+						if (ProtocolTypeClass)
+						{
+							Capture->ProtocolSettings = NewObject<UMovieSceneCapture>(Capture, ProtocolTypeClass);
+							if (Capture->ProtocolSettings)
+							{
+								auto ProtocolDataField = RootObject->TryGetField(TEXT("ProtocolData"));
+								if (ProtocolDataField.IsValid())
+								{
+									FJsonObjectConverter::JsonAttributesToUStruct(ProtocolDataField->AsObject()->Values, ProtocolTypeClass, Capture->ProtocolSettings, 0, 0);
+								}
+							}
+						}
+					}
 				}
 			}
 		}
@@ -128,17 +199,21 @@ private:
 		}
 
 		check( Capture != nullptr );
-		FActiveMovieSceneCaptures::Get().Add( Capture );
 		StartupMovieCaptureHandle = Capture->GetHandle();
+		// Add it immediately, so we can get back to it from its handle (usually it gets added in Initialize())
+		FActiveMovieSceneCaptures::Get().Add( Capture );
+
+		Capture->OnCaptureFinished().AddLambda([]{
+			FPlatformMisc::RequestExit(0);
+		});
+		
 		return Capture;
 	}
 
-	virtual IMovieSceneCaptureInterface* CreateMovieSceneCapture(TWeakPtr<FSceneViewport> InSceneViewport) override
+	virtual IMovieSceneCaptureInterface* CreateMovieSceneCapture(TSharedPtr<FSceneViewport> InSceneViewport) override
 	{
 		UMovieSceneCapture* Capture = NewObject<UMovieSceneCapture>(GetTransientPackage());
-		FActiveMovieSceneCaptures::Get().Add(Capture);
 		Capture->Initialize(InSceneViewport);
-		Capture->SetupFrameRange();
 		Capture->StartCapture();
 		return Capture;
 	}
@@ -176,6 +251,13 @@ private:
 			}
 		}
 	}
+
+	virtual void DestroyAllActiveCaptures()
+	{
+		FCoreDelegates::OnPreExit.RemoveAll(this);
+		PreExit();
+	}
+
 };
 
 IMPLEMENT_MODULE( FMovieSceneCaptureModule, MovieSceneCapture )
diff --git a/Engine/Source/Runtime/MovieSceneCapture/Private/MovieSceneCaptureProtocol.cpp b/Engine/Source/Runtime/MovieSceneCapture/Private/MovieSceneCaptureProtocol.cpp
new file mode 100644
--- /dev/null
+++ b/Engine/Source/Runtime/MovieSceneCapture/Private/MovieSceneCaptureProtocol.cpp
@@ -0,0 +1,14 @@
+// Copyright 1998-2015 Epic Games, Inc. All Rights Reserved.
+
+#include "MovieSceneCapturePCH.h"
+#include "IMovieSceneCaptureProtocol.h"
+#include "SceneViewport.h"
+
+FCaptureProtocolInitSettings FCaptureProtocolInitSettings::FromSlateViewport(TSharedRef<FSceneViewport> InSceneViewport, UObject* InProtocolSettings)
+{
+	FCaptureProtocolInitSettings Settings;
+	Settings.SceneViewport = InSceneViewport;
+	Settings.DesiredSize = InSceneViewport->GetSize();
+	Settings.ProtocolSettings = InProtocolSettings;
+	return Settings;
+}
\ No newline at end of file
diff --git a/Engine/Source/Runtime/MovieSceneCapture/Private/VideoCaptureProtocol.cpp b/Engine/Source/Runtime/MovieSceneCapture/Private/VideoCaptureProtocol.cpp
new file mode 100644
--- /dev/null
+++ b/Engine/Source/Runtime/MovieSceneCapture/Private/VideoCaptureProtocol.cpp
@@ -0,0 +1,75 @@
+// Copyright 1998-2015 Epic Games, Inc. All Rights Reserved.
+
+#include "MovieSceneCapturePCH.h"
+#include "VideoCaptureProtocol.h"
+#include "MovieSceneCapture.h"
+
+bool FVideoCaptureProtocol::Initialize(const FCaptureProtocolInitSettings& InSettings, const ICaptureProtocolHost& Host)
+{
+	if (!FFrameGrabberProtocol::Initialize(InSettings, Host))
+	{
+		return false;
+	}
+
+	UVideoCaptureSettings* CaptureSettings = CastChecked<UVideoCaptureSettings>(InSettings.ProtocolSettings);
+
+#if PLATFORM_MAC
+	static const TCHAR* Extension = TEXT(".mov");
+#else
+	static const TCHAR* Extension = TEXT(".avi");
+#endif
+
+	FAVIWriterOptions Options;
+	Options.OutputFilename = Host.GenerateFilename(FFrameMetrics(), Extension);
+	Options.CaptureFPS = Host.GetCaptureFrequency();
+	Options.CodecName = CaptureSettings->VideoCodec;
+	Options.bSynchronizeFrames = Host.GetCaptureStrategy().ShouldSynchronizeFrames();
+	Options.Width = InSettings.DesiredSize.X;
+	Options.Height = InSettings.DesiredSize.Y;
+
+	if (CaptureSettings->bUseCompression)
+	{
+		Options.CompressionQuality = CaptureSettings->CompressionQuality;
+		
+		float QualityOverride = 100.f;
+		if (FParse::Value( FCommandLine::Get(), TEXT( "-MovieQuality=" ), QualityOverride ))
+		{
+			Options.CompressionQuality = FMath::Clamp(QualityOverride, 1.f, 100.f) / 100.f;
+		}
+
+		Options.CompressionQuality = FMath::Clamp<float>(Options.CompressionQuality.GetValue(), 0.f, 1.f);
+	}
+
+	AVIWriter.Reset(FAVIWriter::CreateInstance(Options));
+	AVIWriter->Initialize();
+
+	return AVIWriter->IsCapturing();
+}
+
+struct FVideoFrameData : IFramePayload
+{
+	FFrameMetrics Metrics;
+};
+
+FFramePayloadPtr FVideoCaptureProtocol::GetFramePayload(const FFrameMetrics& FrameMetrics, const ICaptureProtocolHost& Host) const
+{
+	TSharedRef<FVideoFrameData, ESPMode::ThreadSafe> FrameData = MakeShareable(new FVideoFrameData);
+	FrameData->Metrics = FrameMetrics;
+	return FrameData;
+}
+
+void FVideoCaptureProtocol::ProcessFrame(FCapturedFrameData Frame)
+{
+	FVideoFrameData* Payload = Frame.GetPayload<FVideoFrameData>();
+
+	AVIWriter->DropFrames(Payload->Metrics.NumDroppedFrames);
+	AVIWriter->Update(Payload->Metrics.TotalElapsedTime, MoveTemp(Frame.ColorBuffer));
+}
+
+void FVideoCaptureProtocol::Finalize()
+{
+	AVIWriter->Finalize();
+	AVIWriter.Reset();
+
+	FFrameGrabberProtocol::Finalize();
+}
\ No newline at end of file
diff --git a/Engine/Source/Runtime/MovieSceneCapture/Public/AutomatedLevelSequenceCapture.h b/Engine/Source/Runtime/MovieSceneCapture/Public/AutomatedLevelSequenceCapture.h
--- a/Engine/Source/Runtime/MovieSceneCapture/Public/AutomatedLevelSequenceCapture.h
+++ b/Engine/Source/Runtime/MovieSceneCapture/Public/AutomatedLevelSequenceCapture.h
@@ -21,48 +21,43 @@ public:
 	/** Set the level sequence asset that we are to record. We will spawn a new actor at runtime for this asset for playback. */
 	void SetLevelSequenceAsset(FString InAssetPath);
 
-	/** Set the level sequence actor that we are to record */
-	void SetLevelSequenceActor(ALevelSequenceActor* InActor);
-
 	/** When enabled, the StartFrame setting will override the default starting frame number */
-	UPROPERTY(config, EditAnywhere, Category=CaptureSettings, AdvancedDisplay)
+	UPROPERTY(config, EditAnywhere, Category=Animation, AdvancedDisplay)
 	bool bUseCustomStartFrame;
 
 	/** Frame number to start capturing.  The frame number range depends on whether the bUseRelativeFrameNumbers option is enabled. */
-	UPROPERTY(config, EditAnywhere, Category=CaptureSettings, AdvancedDisplay, meta=(EditCondition="bUseCustomStartFrame"))
+	UPROPERTY(config, EditAnywhere, Category=Animation, AdvancedDisplay, meta=(EditCondition="bUseCustomStartFrame"))
 	int32 StartFrame;
 
 	/** When enabled, the EndFrame setting will override the default ending frame number */
-	UPROPERTY(config, EditAnywhere, Category=CaptureSettings, AdvancedDisplay)
+	UPROPERTY(config, EditAnywhere, Category=Animation, AdvancedDisplay)
 	bool bUseCustomEndFrame;
 
 	/** Frame number to end capturing.  The frame number range depends on whether the bUseRelativeFrameNumbers option is enabled. */
-	UPROPERTY(config, EditAnywhere, Category=CaptureSettings, AdvancedDisplay, meta=(EditCondition="bUseCustomEndFrame"))
+	UPROPERTY(config, EditAnywhere, Category=Animation, AdvancedDisplay, meta=(EditCondition="bUseCustomEndFrame"))
 	int32 EndFrame;
 
 	/** The number of extra frames to play before the sequence's start frame, to "warm up" the animation.  This is useful if your
 	    animation contains particles or other runtime effects that are spawned into the scene earlier than your capture start frame */
-	UPROPERTY(config, EditAnywhere, Category=CaptureSettings, AdvancedDisplay)
+	UPROPERTY(config, EditAnywhere, Category=Animation, AdvancedDisplay)
 	int32 WarmUpFrameCount;
 
 	/** The number of seconds to wait (in real-time) before we start playing back the warm up frames.  Useful for allowing post processing effects to settle down before capturing the animation. */
-	UPROPERTY(config, EditAnywhere, Category=CaptureSettings, AdvancedDisplay, meta=(Units=Seconds, ClampMin=0))
+	UPROPERTY(config, EditAnywhere, Category=Animation, AdvancedDisplay, meta=(Units=Seconds, ClampMin=0))
 	float DelayBeforeWarmUp;
 
 public:
 	// UMovieSceneCapture interface
-	virtual void Initialize(TWeakPtr<FSceneViewport> InViewport) override;
-	virtual void PostEditChangeProperty(FPropertyChangedEvent& PropertyChangedEvent) override;
-	virtual void SetupFrameRange() override;
+	virtual void Initialize(TSharedPtr<FSceneViewport> InViewport, int32 PIEInstance = -1) override;
 
 private:
 
-	/** Called when this capture has stopped */
-	virtual void OnCaptureStopped() override;
-
 	/** Called when the level sequence has updated the world */
 	void SequenceUpdated(const ULevelSequencePlayer& Player, float CurrentTime, float PreviousTime);
 
+	/** Called to set up the player's playback range */
+	void SetupFrameRange();
+
 	/** Delegate binding for the above callback */
 	FDelegateHandle OnPlayerUpdatedBinding;
 
@@ -76,16 +71,12 @@ private:
 
 	/** The pre-existing level sequence actor to use for capture that specifies playback settings */
 	UPROPERTY()
-	TLazyObjectPtr<ALevelSequenceActor> LevelSequenceActor;
-
-	/** Copy of the level sequence ID from Level Sequence. Required because JSON serialization exports the path of the object, rather that its GUID */
-	UPROPERTY()
-	FGuid LevelSequenceActorId;
+	TWeakObjectPtr<ALevelSequenceActor> LevelSequenceActor;
 
 	/** Which state we're in right now */
 	enum class ELevelSequenceCaptureState
 	{
-		Staging,
+		Setup,
 		DelayBeforeWarmUp,
 		ReadyToWarmUp,
 		WarmingUp,
@@ -97,6 +88,9 @@ private:
 
 	/** The number of warm up frames left before we actually start saving out images */
 	int32 RemainingWarmUpFrames;
+
+	/** The delta of the last sequence updat */
+	float LastSequenceUpdateDelta;
 	
 #endif
 };
diff --git a/Engine/Source/Runtime/MovieSceneCapture/Public/FrameGrabber.h b/Engine/Source/Runtime/MovieSceneCapture/Public/FrameGrabber.h
new file mode 100644
--- /dev/null
+++ b/Engine/Source/Runtime/MovieSceneCapture/Public/FrameGrabber.h
@@ -0,0 +1,199 @@
+// Copyright 1998-2015 Epic Games, Inc. All Rights Reserved.
+
+#pragma once
+
+class FSceneViewport;
+
+/** A single, managed surface used as a render target resolution destination  */
+struct MOVIESCENECAPTURE_API FViewportSurfaceReader
+{
+	/** Constructor */
+	FViewportSurfaceReader(EPixelFormat InPixelFormat, FIntPoint InBufferSize);
+
+	/** Destructor */
+	~FViewportSurfaceReader();
+
+	/** Initialize this reader so that it can be waited on. */
+	void Initialize();
+
+	/** Wait for this reader to become available, if it's currently in use */
+	void BlockUntilAvailable();
+
+	/**
+	 * Resolve the specified viewport RHI, calling the specified callback with the result.
+	 *
+	 * @param	ViewportRHI		The viewport to resolve
+	 * @param	Callback 		Callback to call with the locked texture data. This will be called on an undefined thread.
+	 */
+	void ResolveRenderTarget(const FViewportRHIRef& ViewportRHI, TFunction<void(FColor*)> Callback);
+
+	/** Get the current size of the texture */
+	FIntPoint GetCurrentSize() const;
+
+	/** Set the rectangle within which to read pixels */
+	void SetCaptureRect(FIntRect InCaptureRect) { CaptureRect = InCaptureRect; }
+
+protected:
+
+	/** Set up this surface to the specified width/height */
+	void Resize(uint32 Width, uint32 Height);
+
+	/** Whether this surface reader is enabled or not */
+	FThreadSafeBool bEnabled;
+
+	/** Optional event that is triggered when the surface is no longer in use */
+	FEvent* AvailableEvent;
+
+	/** Texture used to store the resolved render target */
+	FTexture2DRHIRef ReadbackTexture;
+
+	/** The rectangle to read from the surface */
+	FIntRect CaptureRect;
+
+	/** The desired pixel format of the resolved textures */
+	EPixelFormat PixelFormat;
+
+	/** Whether this reader is enabled or not. */
+	bool bIsEnabled;
+};
+
+struct IFramePayload
+{
+	virtual ~IFramePayload() {}
+};
+
+typedef TSharedPtr<IFramePayload, ESPMode::ThreadSafe> FFramePayloadPtr;
+
+/** Structure representing a captured frame */
+struct FCapturedFrameData
+{
+	FCapturedFrameData(FIntPoint InBufferSize, FFramePayloadPtr InPayload) : BufferSize(InBufferSize), Payload(MoveTemp(InPayload)) {}
+
+	FCapturedFrameData(FCapturedFrameData&& In) : ColorBuffer(MoveTemp(In.ColorBuffer)), BufferSize(In.BufferSize), Payload(MoveTemp(In.Payload)) {}
+	FCapturedFrameData& operator=(FCapturedFrameData&& In){ ColorBuffer = MoveTemp(In.ColorBuffer); BufferSize = In.BufferSize; Payload = MoveTemp(In.Payload); return *this; }
+
+	template<typename T>
+	T* GetPayload() { return static_cast<T*>(Payload.Get()); }
+
+	/** The color buffer of the captured frame */
+	TArray<FColor> ColorBuffer;
+
+	/** The size of the resulting color buffer */
+	FIntPoint BufferSize;
+
+	/** Optional user-specified payload */
+	FFramePayloadPtr Payload;
+	
+private:
+	FCapturedFrameData(const FCapturedFrameData& In);
+	FCapturedFrameData& operator=(const FCapturedFrameData& In);
+};
+
+/**
+ * Class responsible for resolving render target data for a specific viewport in an efficient manner
+ * Internally, the class uses a fixed array of resolution surfaces, and dispatches rendering commands
+ * to resolve the viewport render target into a specific index into this array. This means we can
+ * resolve the render target data without having to wait, or flush rendering commands.
+ */
+class MOVIESCENECAPTURE_API FFrameGrabber
+{
+public:
+	/**
+	 * Construct this frame grabber
+	 *
+	 * @param InViewport		The viewport we are to grab frames for
+	 * @param DesiredBufferSize	The desired size of captured frames
+	 * @param InPixelFormat		The desired pixel format to store captured frames as
+	 * @param InNumSurfaces		The number of destination surfaces contained in our buffer 
+	 */
+	FFrameGrabber(TSharedRef<FSceneViewport> Viewport, FIntPoint DesiredBufferSize, EPixelFormat InPixelFormat = PF_B8G8R8A8, uint32 NumSurfaces = 2);
+
+	/** Destructor */
+	~FFrameGrabber();
+
+public:
+
+	/** Instruct the frame grabber to start capturing frames */
+	void StartCapturingFrames();
+
+	/** Instruct the frame grabber capture this frame, when it receives an event from slate */
+	void CaptureThisFrame(FFramePayloadPtr Payload);
+
+	/** Stop capturing frames */
+	void StopCapturingFrames();
+
+	/** Shut down this grabber, ensuring that any threaded operations are finished */
+	void Shutdown();
+
+public:
+
+	/** Check whether we have any outstanding frames or not */
+	bool HasOutstandingFrames() const;
+
+	/** Retrieve any frames we may have captured */
+	TArray<FCapturedFrameData> GetCapturedFrames();
+
+protected:
+	
+	/** Callback for when slate has rendered a window (called on game thread) */
+	void OnSlateWindowRendered( SWindow& SlateWindow, void* ViewportRHIPtr );
+
+	/** Called when the specified surface index has been locked for reading with the render target data (called on render thread)  */
+	void OnFrameReady(int32 SurfaceIndex, FColor* ColorBuffer);
+
+private:
+
+	/** Non-copyable */
+	FFrameGrabber(const FFrameGrabber&);
+	FFrameGrabber& operator=(const FFrameGrabber&);
+
+	/** The window we'll capture, and the capture rectangle */
+	TWeakPtr<SWindow> CaptureWindow;
+
+	/** Delegate handle for the OnSlateOnWindowRendered event */
+	FDelegateHandle OnWindowRendered;
+
+	/** Array of captured frames */
+	TArray<FCapturedFrameData> CapturedFrames;
+
+	/** Lock to protect the above array */
+	mutable FCriticalSection CapturedFramesMutex;
+
+	/** Array of surfaces that we resolve the viewport RHI to. Fixed allocation - should never be resized */
+	struct FResolveSurface
+	{
+		FResolveSurface(EPixelFormat InPixelFormat, FIntPoint BufferSize) : Surface(InPixelFormat, BufferSize) {}
+
+		FFramePayloadPtr Payload;
+		FViewportSurfaceReader Surface;
+	};
+	TArray<FResolveSurface> Surfaces;
+
+	/** Index into the above array to the next surface that we should use - only accessed on main thread */
+	int32 CurrentFrameIndex;
+
+	/** The index that we should capture the next rendered slate window into */
+	int32 SlateRenderIndex;
+
+	/** The total number of frames we are currently waiting on */
+	FThreadSafeCounter OutstandingFrameCount;
+
+	/** Pending frame payloads to be passed with frames captured from slate  */
+	FCriticalSection PendingFramePayloadsMutex;
+	TArray<FFramePayloadPtr> PendingFramePayloads;
+
+	/** When CaptureThisFrame is called, we issue a rendering command to ensure that the frame we capture is the frame that was actually requested.
+	 * This is required because the game thread and render thread are not necessarily in-sync
+	 */
+	FThreadSafeCounter LatentRenderCommandCount;
+
+	/** The current state of the grabber */
+	enum class EFrameGrabberState
+	{
+		Inactive, Active, PendingShutdown
+	};
+	EFrameGrabberState State;
+
+	/** The desired target size to resolve frames to */
+	FIntPoint TargetSize;
+};
\ No newline at end of file
diff --git a/Engine/Source/Runtime/MovieSceneCapture/Public/IMovieSceneCapture.h b/Engine/Source/Runtime/MovieSceneCapture/Public/IMovieSceneCapture.h
--- a/Engine/Source/Runtime/MovieSceneCapture/Public/IMovieSceneCapture.h
+++ b/Engine/Source/Runtime/MovieSceneCapture/Public/IMovieSceneCapture.h
@@ -22,12 +22,10 @@ public:
 	GENERATED_BODY()
 
 	/** Initialize this capture object by binding it to the specified viewport */
-	virtual void Initialize(TWeakPtr<FSceneViewport> Viewport) = 0;
+	virtual void Initialize(TSharedPtr<FSceneViewport> Viewport, int32 PIEInstance = -1) = 0;
 
-	/** Called after initialization to allow the movie scene capture object to override the range of frames to 
-	    capture.  It's called after initialization so you can be sure that the FrameRate and other variables are
-		fully initialized first. */
-	virtual void SetupFrameRange() = 0;
+	/** Instruct this capture to start capturing frames */
+	virtual void StartCapturing() = 0;
 
 	/** Shut down this movie capture */
 	virtual void Close() = 0;
diff --git a/Engine/Source/Runtime/MovieSceneCapture/Public/IMovieSceneCaptureProtocol.h b/Engine/Source/Runtime/MovieSceneCapture/Public/IMovieSceneCaptureProtocol.h
new file mode 100644
--- /dev/null
+++ b/Engine/Source/Runtime/MovieSceneCapture/Public/IMovieSceneCaptureProtocol.h
@@ -0,0 +1,130 @@
+// Copyright 1998-2015 Epic Games, Inc. All Rights Reserved.
+
+#pragma once
+
+class FSceneViewport;
+
+struct ICaptureProtocolHost;
+struct FFrameMetrics;
+struct FCaptureProtocolInitSettings;
+
+/** A capture protocol responsible for dealing with captured frames using some custom method (writing out to disk, streaming, etc) */
+struct IMovieSceneCaptureProtocol
+{
+	/**
+	 * Initialize this capture protocol
+	 *
+	 * @param InSettings		The initial initialization settings to use for the capture
+	 * @param Host				The client that is initializing this protocol
+	 *
+	 * @return true if initialization was successful, false otherwise
+	 */
+	virtual bool Initialize(const FCaptureProtocolInitSettings& InSettings, const ICaptureProtocolHost& Host) = 0;
+
+	/**
+	 * Instruct this protocol to capture a frame relating to the specified metrics
+	 *
+	 * @param FrameMetrics		Frame metrics relating to the current frame
+	 * @param Host				The client that is initializing this protocol
+	 */
+	virtual void CaptureFrame(const FFrameMetrics& FrameMetrics, const ICaptureProtocolHost& Host) = 0;
+
+	/**
+	 * Check whether this protocol has any processing left to do, or whether it should be finalized.
+	 * Only called when the capture has been asked to end.
+	 */
+	virtual bool HasFinishedProcessing() const { return true; }
+
+	/**
+	 * Called on the main thread to do any additional processing
+	 */
+	virtual void Tick() {}
+
+	/**
+	 * Called when we have finished capturing
+	 */
+	virtual void Finalize() {}
+
+	/**
+	 * Called when generating formatting filename to add any additional format mappings
+	 *
+	 * @param FormatMappings	Map to add additional format rules to
+	 */
+	virtual void AddFormatMappings(TMap<FString, FStringFormatArg>& FormatMappings) const {}
+
+	virtual ~IMovieSceneCaptureProtocol() {}
+};
+
+
+
+/** Metrics that correspond to a particular frame */
+struct FFrameMetrics
+{
+	/** Default construction */
+	FFrameMetrics() : TotalElapsedTime(0), FrameDelta(0), FrameNumber(0), NumDroppedFrames(0) {}
+
+	FFrameMetrics(float InTotalElapsedTime, float InFrameDelta, uint32 InFrameNumber, uint32 InNumDroppedFrames)
+		: TotalElapsedTime(InTotalElapsedTime), FrameDelta(InFrameDelta), FrameNumber(InFrameNumber), NumDroppedFrames(InNumDroppedFrames)
+	{
+	}
+
+	/** The total amount of time, in seconds, since the capture started */
+	float TotalElapsedTime;
+	/** The total amount of time, in seconds, that this specific frame took to render (not accounting for dropped frames) */
+	float FrameDelta;
+	/** The index of this frame from the start of the capture, including dropped frames */
+	uint32 FrameNumber;
+	/** The number of frames we dropped in-between this frame, and the last one we captured */
+	uint32 NumDroppedFrames;
+};
+
+
+/** Interface that defines when to capture or drop frames */
+struct ICaptureStrategy
+{
+	virtual ~ICaptureStrategy(){}
+	
+	virtual void OnWarmup() = 0;
+	virtual void OnStart() = 0;
+	virtual void OnStop() = 0;
+	virtual void OnPresent(double CurrentTimeSeconds, uint32 FrameIndex) = 0;
+	virtual bool ShouldSynchronizeFrames() const { return true; }
+
+	virtual bool ShouldPresent(double CurrentTimeSeconds, uint32 FrameIndex) const = 0;
+	virtual int32 GetDroppedFrames(double CurrentTimeSeconds, uint32 FrameIndex) const = 0;
+};
+
+
+/** Interface to be implemented by any class using an IMovieSceneCaptureProtocol instance*/
+struct ICaptureProtocolHost
+{
+	/** Generate a filename for the specified frame metrics. How often this is called is determined by the protocol itself */
+	virtual FString GenerateFilename(const FFrameMetrics& FrameMetrics, const TCHAR* Extension) const = 0;
+
+	/** Get the capture frequency */
+	virtual float GetCaptureFrequency() const = 0;
+
+	/** Access the host's capture strategy */
+	virtual const ICaptureStrategy& GetCaptureStrategy() const = 0;
+};
+
+
+/** Structure used to initialize a capture protocol */
+struct FCaptureProtocolInitSettings
+{
+	/**~ @todo: add ability to capture a sub-rectangle */
+
+	/** Capture from a slate viewport, using the specified custom protocol settings */
+	MOVIESCENECAPTURE_API static FCaptureProtocolInitSettings FromSlateViewport(TSharedRef<FSceneViewport> InSceneViewport, UObject* InProtocolSettings);
+	
+	/** The slate viewport we should capture from */
+	TSharedPtr<FSceneViewport> SceneViewport;
+	/** The desired size of the captured frames */
+	FIntPoint DesiredSize;
+	/** Protocol settings specific to the protocol we will be initializing */
+	UObject* ProtocolSettings;
+
+private:
+	/** Private construction to ensure use of static init methods */
+	FCaptureProtocolInitSettings() : ProtocolSettings(nullptr) {}
+};
\ No newline at end of file
diff --git a/Engine/Source/Runtime/MovieSceneCapture/Public/LevelCapture.h b/Engine/Source/Runtime/MovieSceneCapture/Public/LevelCapture.h
--- a/Engine/Source/Runtime/MovieSceneCapture/Public/LevelCapture.h
+++ b/Engine/Source/Runtime/MovieSceneCapture/Public/LevelCapture.h
@@ -9,24 +9,29 @@ UCLASS()
 class MOVIESCENECAPTURE_API ULevelCapture : public UMovieSceneCapture
 {
 public:
+	ULevelCapture(const FObjectInitializer&);
+
 	GENERATED_BODY()
 
+	/** Specifies whether the capture should start immediately, or whether it will be invoked externally (through StartMovieCapture/StopMovieCapture exec commands) */
+	UPROPERTY(EditAnywhere, Category=General)
+	bool bAutoStartCapture;
+
 	/** Specify a prerequisite actor that must be set up before we start capturing */
 	void SetPrerequisiteActor(AActor* Prereq);
 
-	virtual void Initialize(TWeakPtr<FSceneViewport> InViewport) override;
+	virtual void Initialize(TSharedPtr<FSceneViewport> InViewport, int32 PIEInstance = -1) override;
 	virtual void Tick(float DeltaSeconds) override;
-#if WITH_EDITOR
-	virtual void PostEditChangeProperty(FPropertyChangedEvent& PropertyChangedEvent) override;
-#endif
 	
 private:
 
 	/** Specify a prerequisite actor that must be set up before we start capturing */
-	UPROPERTY()
-	TLazyObjectPtr<AActor> PrerequisiteActor;
+	TWeakObjectPtr<AActor> PrerequisiteActor;
 
 	/** Copy of the ID from PrerequisiteActor. Required because JSON serialization exports the path of the object, rather that its GUID */
 	UPROPERTY()
 	FGuid PrerequisiteActorId;
+
+	/** PIE instance index that we are capturing with */
+	int32 PIECaptureInstance;
 };
\ No newline at end of file
diff --git a/Engine/Source/Runtime/MovieSceneCapture/Public/MovieSceneCapture.h b/Engine/Source/Runtime/MovieSceneCapture/Public/MovieSceneCapture.h
--- a/Engine/Source/Runtime/MovieSceneCapture/Public/MovieSceneCapture.h
+++ b/Engine/Source/Runtime/MovieSceneCapture/Public/MovieSceneCapture.h
@@ -6,25 +6,11 @@
 #include "IMovieSceneCapture.h"
 #include "MovieSceneCaptureHandle.h"
 #include "MovieScene.h"
-#include "AVIWriter.h"
-#include "Tickable.h"
+#include "RenderingThread.h"
+#include "IMovieSceneCaptureProtocol.h"
+#include "MovieSceneCaptureProtocolRegistry.h"
 #include "MovieSceneCapture.generated.h"
 
-/** Interface that defines when to capture or drop frames */
-struct ICaptureStrategy
-{
-	virtual ~ICaptureStrategy(){}
-	
-	virtual void OnWarmup() = 0;
-	virtual void OnStart() = 0;
-	virtual void OnStop() = 0;
-	virtual void OnPresent(double CurrentTimeSeconds, uint32 FrameIndex) = 0;
-	virtual bool ShouldSynchronizeFrames() const { return true; }
-
-	virtual bool ShouldPresent(double CurrentTimeSeconds, uint32 FrameIndex) const = 0;
-	virtual int32 GetDroppedFrames(double CurrentTimeSeconds, uint32 FrameIndex) const = 0;
-};
-
 /** Structure used to cache various metrics for our capture */
 struct FCachedMetrics
 {
@@ -40,142 +26,132 @@ struct FCachedMetrics
 
 /** Class responsible for capturing scene data */
 UCLASS(config=EditorPerProjectUserSettings)
-class MOVIESCENECAPTURE_API UMovieSceneCapture : public UObject, public IMovieSceneCaptureInterface
+class MOVIESCENECAPTURE_API UMovieSceneCapture : public UObject, public IMovieSceneCaptureInterface, public ICaptureProtocolHost
 {
 public:
-
 	UMovieSceneCapture(const FObjectInitializer& Initializer);
 
 	GENERATED_BODY()
 
+	virtual void PostInitProperties() override;
+
 public:
 
 	// Begin IMovieSceneCaptureInterface
-	virtual void Initialize(TWeakPtr<FSceneViewport> InSceneViewport) override;
-	virtual void SetupFrameRange() override
-	{
-	}
-	virtual void Close() override;
+	virtual void Initialize(TSharedPtr<FSceneViewport> InSceneViewport, int32 PIEInstance = -1) override;
+	virtual void StartCapturing() { StartCapture(); }
+	virtual void Close() override { Finalize(); }
 	virtual FMovieSceneCaptureHandle GetHandle() const override { return Handle; }
 	const FMovieSceneCaptureSettings& GetSettings() const override { return Settings; }
 	// End IMovieSceneCaptureInterface
 
 public:
 
+	/** The type of capture protocol to use */
+	UPROPERTY(config, EditAnywhere, Category=CaptureSettings)
+	FCaptureProtocolID CaptureType;
+
+	/** Settings specific to the capture protocol */
+	UPROPERTY(EditAnywhere, Category=CaptureSettings)
+	UObject* ProtocolSettings;
+
 	/** Settings that define how to capture */
 	UPROPERTY(EditAnywhere, config, Category=CaptureSettings, meta=(ShowOnlyInnerProperties))
 	FMovieSceneCaptureSettings Settings;
 
+	/** Whether to capture the movie in a separate process or not */
+	UPROPERTY(config, EditAnywhere, Category=General, AdvancedDisplay)
+	bool bUseSeparateProcess;
+
 	/** Additional command line arguments to pass to the external process when capturing */
-	UPROPERTY(EditAnywhere, config, Category=General, AdvancedDisplay)
+	UPROPERTY(EditAnywhere, config, Category=General, AdvancedDisplay, meta=(EditCondition=bUseSeparateProcess))
 	FString AdditionalCommandLineArguments;
 
 	/** Command line arguments inherited from this process */
-	UPROPERTY(EditAnywhere, transient, Category=General, AdvancedDisplay)
+	UPROPERTY(EditAnywhere, transient, Category=General, AdvancedDisplay, meta=(EditCondition=bUseSeparateProcess))
 	FString InheritedCommandLineArguments;
 
-	/** Value used to control the BufferVisualizationDumpFrames cvar in the child process */
-	UPROPERTY()
-	bool bBufferVisualizationDumpFrames;
+	/** Event that is fired after we've finished capturing */
+	DECLARE_EVENT( UMovieSceneCapture, FOnCaptureFinished );
+	FOnCaptureFinished& OnCaptureFinished() { return OnCaptureFinishedDelegate; }
 
 public:
 
 	/** Access this object's cached metrics */
 	const FCachedMetrics& GetMetrics() const { return CachedMetrics; }
 
+	/** Access the capture protocol we are using */
+	IMovieSceneCaptureProtocol* GetCaptureProtocol() { return CaptureProtocol.Get(); }
+
+public:
+
 	/** Starts warming up.  May be optionally called before StartCapture().  This can be used to start rendering frames early, before
-	    any files are written out */
+	    any files are captured or written out */
 	void StartWarmup();
 
-	/** Initialize the capture */
+	/** Initialize the capture so that it is able to start capturing frames */
 	void StartCapture();
 
-	/** Finalize the capture, waiting for any outstanding processing */
-	void StopCapture();
+	/** Indicate that this frame should be captured - must be called before the movie scene capture is ticked */
+	void CaptureThisFrame(float DeltaSeconds);
 
-	/** Capture a frame from our previously populated color buffer (populated by slate) */
-	void CaptureFrame(float DeltaSeconds);
+	/** Automatically finalizes the capture when all currently pending frames are dealt with */
+	void FinalizeWhenReady();
 
-protected:
+	/** Check whether we should automatically finalize this capture */
+	bool ShouldFinalize() const { return bFinalizeWhenReady && CaptureProtocol->HasFinishedProcessing(); }
 
-#if WITH_EDITOR
-	/** Implementation function that saves out a snapshot file from the specified color data */
-	void SaveFrameToFile(TArray<FColor> Colors);
-#endif
+	/** Finalize the capturing process, assumes all frames have been processed. */
+	void Finalize();
 
-	/** Prepare the slate renderer for a screenshot */
-	void PrepareForScreenshot();
+public:
 
 	/** Called at the end of a frame, before a frame is presented by slate */
-	virtual void Tick(float DeltaSeconds);
-
-private:
-
-	/** Called when this movie scene capture has stopped */
-	virtual void OnCaptureStopped() {}
+	virtual void Tick(float DeltaSeconds) { CaptureThisFrame(DeltaSeconds); }
 
 protected:
 
-	/** Resolve the specified format using the user supplied formatting rules. */
-	FString ResolveFileFormat(const FString& Folder, const FString& Format) const;
-
-	/** Get the default extension required for this capture (avi, jpg etc) */
-	const TCHAR* GetDefaultFileExtension() const;
-
-	/** Calculate a unique index for the {unique} formatting rule */
-	FString ResolveUniqueFilename();
+	/**~ ICaptureProtocolHost interface */
+	virtual FString GenerateFilename(const FFrameMetrics& FrameMetrics, const TCHAR* Extension) const override;
+	virtual float GetCaptureFrequency() const { return Settings.FrameRate; }
+	virtual const ICaptureStrategy& GetCaptureStrategy() const { return *CaptureStrategy; }
 
-private:
-
-	struct FTicker : public FTickableGameObject
-	{
-		FTicker(UMovieSceneCapture* InCapture) : Capture(InCapture) {}
+	/** Resolve the specified format using the user supplied formatting rules. */
+	FString ResolveFileFormat(const FString& Format, const FFrameMetrics& FrameMetrics) const;
 
-	private:
-		virtual bool IsTickableInEditor() const override { return false; }
-		virtual bool IsTickable() const override { return true; }
-		virtual bool IsTickableWhenPaused() const override { return false; }
-		virtual TStatId GetStatId() const override { RETURN_QUICK_DECLARE_CYCLE_STAT(UMovieSceneCapture, STATGROUP_Tickables); }
-		virtual void Tick(float DeltaSeconds) override
-		{
-			Capture->Tick(DeltaSeconds);
-		}
+protected:
 
-		UMovieSceneCapture* Capture;
-	};
-	friend FTicker;
+#if WITH_EDITOR
+	virtual void PostEditChangeProperty( struct FPropertyChangedEvent& PropertyChangedEvent) override;
+#endif
 
 protected:
-	
-	/** The viewport we are bound to */
-	TWeakPtr<FSceneViewport> SceneViewport;
+	/** Capture protocol responsible for actually capturing frame data */
+	TSharedPtr<IMovieSceneCaptureProtocol> CaptureProtocol;
+	/** Strategy used for capture (real-time/fixed-time-step) */
+	TSharedPtr<ICaptureStrategy> CaptureStrategy;
+	/** The settings we will use to set up the capture protocol */
+	TOptional<FCaptureProtocolInitSettings> InitSettings;
+	/** Whether we should automatically attempt to capture frames every tick or not */
+	bool bFinalizeWhenReady;
 	/** Our unique handle, used for external representation without having to link to the MovieSceneCapture module */
 	FMovieSceneCaptureHandle Handle;
 	/** Cached metrics for this capture operation */
 	FCachedMetrics CachedMetrics;
-	/** Optional AVI file writer used when capturing to video */
-	TUniquePtr<FAVIWriter> AVIWriter;
-	/** Strategy used for capture (real-time/fixed-time-step) */
-	TSharedPtr<ICaptureStrategy> CaptureStrategy;
-	/** Scratch space for per-frame screenshots */
-	TArray<FColor> ScratchBuffer;
+	/** Format mappings used for generating filenames */
+	mutable TMap<FString, FStringFormatArg> FormatMappings;
 	/** The number of frames to capture.  If this is zero, we'll capture the entire sequence. */
 	int32 FrameCount;
-	/** Running count of how many frames we're waiting to capture from slate */
-	int32 OutstandingFrameCount;
-	/** The delta of the last frame */
-	float LastFrameDelta;
 	/** Whether we have started capturing or not */
 	bool bCapturing;
-	/** Class responsible for ticking this instance */
-	TUniquePtr<FTicker> Ticker;
 	/** Frame number index offset when saving out frames.  This is used to allow the frame numbers on disk to match
 	    what they would be in the authoring application, rather than a simple 0-based sequential index */
 	int32 FrameNumberOffset;
-
+	/** Event that is triggered when capturing has finished */
+	FOnCaptureFinished OnCaptureFinishedDelegate;
 };
 
-/** A strategy that employs a fixed frame time-step, and as such never drops a frame. Potentially accelerated.  */
+/** A strategy that employs a fixed frame time-step, and as such never drops a frame. Potentially accelerated. */
 struct MOVIESCENECAPTURE_API FFixedTimeStepCaptureStrategy : ICaptureStrategy
 {
 	FFixedTimeStepCaptureStrategy(uint32 InTargetFPS);
@@ -183,7 +159,7 @@ struct MOVIESCENECAPTURE_API FFixedTimeStepCaptureStrategy : ICaptureStrategy
 	virtual void OnWarmup() override;
 	virtual void OnStart() override;
 	virtual void OnStop() override;
-	virtual void OnPresent(double CurrentTimeSeconds, uint32 FrameIndex) override;	
+	virtual void OnPresent(double CurrentTimeSeconds, uint32 FrameIndex) override;
 	virtual bool ShouldPresent(double CurrentTimeSeconds, uint32 FrameIndex) const override;
 	virtual int32 GetDroppedFrames(double CurrentTimeSeconds, uint32 FrameIndex) const override;
 
diff --git a/Engine/Source/Runtime/MovieSceneCapture/Public/MovieSceneCaptureModule.h b/Engine/Source/Runtime/MovieSceneCapture/Public/MovieSceneCaptureModule.h
--- a/Engine/Source/Runtime/MovieSceneCapture/Public/MovieSceneCaptureModule.h
+++ b/Engine/Source/Runtime/MovieSceneCapture/Public/MovieSceneCaptureModule.h
@@ -6,7 +6,7 @@
 #include "IMovieSceneCapture.h"
 #include "MovieSceneCaptureHandle.h"
 
-class UMovieSceneCapture;
+class FMovieSceneCaptureProtocolRegistry;
 
 class IMovieSceneCaptureModule : public IModuleInterface
 {
@@ -18,10 +18,13 @@ public:
 	}
 
 	virtual IMovieSceneCaptureInterface* InitializeFromCommandLine() = 0;
-	virtual IMovieSceneCaptureInterface* CreateMovieSceneCapture(TWeakPtr<FSceneViewport> Viewport) = 0;
+	virtual IMovieSceneCaptureInterface* CreateMovieSceneCapture(TSharedPtr<FSceneViewport> Viewport) = 0;
 	virtual IMovieSceneCaptureInterface* GetFirstActiveMovieSceneCapture() = 0;
 	
 	virtual IMovieSceneCaptureInterface* RetrieveMovieSceneInterface(FMovieSceneCaptureHandle Handle) = 0;
 	virtual void DestroyMovieSceneCapture(FMovieSceneCaptureHandle Handle) = 0;
+
+	virtual void DestroyAllActiveCaptures() = 0;
+	virtual FMovieSceneCaptureProtocolRegistry& GetProtocolRegistry() = 0;
 };
 
diff --git a/Engine/Source/Runtime/MovieSceneCapture/Public/MovieSceneCaptureProtocolRegistry.h b/Engine/Source/Runtime/MovieSceneCapture/Public/MovieSceneCaptureProtocolRegistry.h
new file mode 100644
--- /dev/null
+++ b/Engine/Source/Runtime/MovieSceneCapture/Public/MovieSceneCaptureProtocolRegistry.h
@@ -0,0 +1,95 @@
+// Copyright 1998-2015 Epic Games, Inc. All Rights Reserved.
+
+#pragma once
+
+#include "MovieSceneCaptureProtocolRegistry.generated.h"
+
+struct IMovieSceneCaptureProtocol;
+
+/** Structure used to uniquely identify a specific capture protocol */
+USTRUCT()
+struct FCaptureProtocolID
+{
+	FCaptureProtocolID(){}
+	FCaptureProtocolID(const TCHAR* InName) : Identifier(InName) {}
+
+	GENERATED_BODY()
+
+	UPROPERTY(VisibleAnywhere, Category=ID)
+	FName Identifier;
+};
+
+inline bool operator==(FCaptureProtocolID A, FCaptureProtocolID B){ return A.Identifier == B.Identifier; }
+inline bool operator!=(FCaptureProtocolID A, FCaptureProtocolID B){ return A.Identifier != B.Identifier; }
+inline uint32 GetTypeHash(FCaptureProtocolID ID){ return GetTypeHash(ID.Identifier); }
+
+/** Structure that defines a particular capture protocol */
+struct FMovieSceneCaptureProtocolInfo
+{
+	/** This protocol's display name */
+	FText DisplayName;
+	/** Factory function called to create a new instance of this protocol */
+	TFunction<TSharedRef<IMovieSceneCaptureProtocol>()> Factory;
+	/** Custom settings class type to use for this protocol */
+	UClass* SettingsClassType;
+};
+
+/** Class responsible for maintaining a list of available capture protocols */
+class FMovieSceneCaptureProtocolRegistry
+{
+public:
+
+	/** Check whether the specified ID corresponds to a valid protocol */
+	bool IsValidProtocol(FCaptureProtocolID InProtocolID)
+	{
+		return Register.Contains(InProtocolID);
+	}
+
+	/** Iterate alll the protocols we are currently aware of */
+	void IterateProtocols(TFunctionRef<void(FCaptureProtocolID, const FMovieSceneCaptureProtocolInfo&)> Iter)
+	{
+		for (auto& Pair : Register)
+		{
+			Iter(Pair.Key, Pair.Value);
+		}
+	}
+
+	/** Create a new settings type for the specified protocol ID */
+	UObject* FactorySettingsType(FCaptureProtocolID InProtocolID, UObject* Outer)
+	{
+		if (const FMovieSceneCaptureProtocolInfo* Info = Register.Find(InProtocolID))
+		{
+			if (Info->SettingsClassType)
+			{
+				return NewObject<UObject>(Outer, Info->SettingsClassType);
+			}
+		}
+		return nullptr;
+	}
+
+	/** Create a new instance of the protocol that relates to the specified ID */
+	TSharedPtr<IMovieSceneCaptureProtocol> Factory(FCaptureProtocolID InProtocolID)
+	{
+		const FMovieSceneCaptureProtocolInfo* Info = Register.Find(InProtocolID);
+		if (Info)
+		{
+			return Info->Factory();
+		}
+		return nullptr;
+	}
+
+	/** Register a new protocol */
+	void RegisterProtocol(FCaptureProtocolID InProtocolID, const FMovieSceneCaptureProtocolInfo& Info)
+	{
+		Register.Add(InProtocolID, Info);
+	}
+
+	/** Unregister a previously registered protocol */
+	void UnRegisterProtocol(FCaptureProtocolID InProtocolID)
+	{
+		Register.Remove(InProtocolID);
+	}
+
+private:
+	TMap<FCaptureProtocolID, FMovieSceneCaptureProtocolInfo> Register;
+};
\ No newline at end of file
diff --git a/Engine/Source/Runtime/MovieSceneCapture/Public/MovieSceneCaptureSettings.h b/Engine/Source/Runtime/MovieSceneCapture/Public/MovieSceneCaptureSettings.h
--- a/Engine/Source/Runtime/MovieSceneCapture/Public/MovieSceneCaptureSettings.h
+++ b/Engine/Source/Runtime/MovieSceneCapture/Public/MovieSceneCaptureSettings.h
@@ -5,14 +5,6 @@
 #include "Engine/EngineTypes.h"
 #include "MovieSceneCaptureSettings.generated.h"
 
-UENUM()
-enum class EMovieCaptureType : uint8
-{
-	AVI		UMETA(DisplayName="AVI Movie"),
-	BMP		UMETA(DisplayName="BMP Image Sequence"),
-	PNG		UMETA(DisplayName="PNG Image Sequence"),
-	JPEG	UMETA(DisplayName="JPEG Image Sequence")
-};
 
 /** Structure representing a capture resolution */
 USTRUCT()
@@ -75,22 +67,6 @@ struct MOVIESCENECAPTURE_API FMovieSceneCaptureSettings
 	/** The resolution at which to capture */
 	UPROPERTY(config, EditAnywhere, Category=CaptureSettings, meta=(ShowOnlyInnerProperties))
 	FCaptureResolution Resolution;
-
-	/** The type of capture to perform */
-	UPROPERTY(config, EditAnywhere, Category=CaptureSettings)
-	EMovieCaptureType CaptureType;
-
-	/** Whether compression is enabled on the resulting file(s) */
-	UPROPERTY(config, EditAnywhere, Category=CaptureSettings, AdvancedDisplay )
-	bool bUseCompression;
-
-	/** For output formats that support compression, this specifies the quality level of the compressed image.  0.0 results in the most compression, while 1.0 gives you the highest image quality. */
-	UPROPERTY(config, EditAnywhere, Category=CaptureSettings, AdvancedDisplay, meta=(ClampMin=0, ClampMax=1), meta=(EditCondition="bUseCompression"))
-	float CompressionQuality;
-
-	/** (Experimental) - An optional codec to use for video encoding */
-	UPROPERTY(config, EditAnywhere, Category=CaptureSettings, AdvancedDisplay)
-	FString Codec;
 	
 	/** Whether to texture streaming should be enabled while capturing.  Turning off texture streaming may cause much more memory to be used, but also reduces the chance of blurry textures in your captured video. */
 	UPROPERTY(config, EditAnywhere, Category=CaptureSettings, AdvancedDisplay)
diff --git a/Engine/Source/Runtime/MovieSceneCapture/Public/Protocols/CompositionGraphCaptureProtocol.h b/Engine/Source/Runtime/MovieSceneCapture/Public/Protocols/CompositionGraphCaptureProtocol.h
new file mode 100644
--- /dev/null
+++ b/Engine/Source/Runtime/MovieSceneCapture/Public/Protocols/CompositionGraphCaptureProtocol.h
@@ -0,0 +1,55 @@
+// Copyright 1998-2015 Epic Games, Inc. All Rights Reserved.
+
+#pragma once
+
+#include "IMovieSceneCaptureProtocol.h"
+#include "CompositionGraphCaptureProtocol.generated.h"
+
+USTRUCT()
+struct MOVIESCENECAPTURE_API FCompositionGraphCapturePasses
+{
+	GENERATED_BODY()
+
+	UPROPERTY(EditAnywhere, Category="Composition Graph Settings")
+	TArray<FString> Value;
+};
+
+UCLASS(config=EditorPerProjectUserSettings, DisplayName="Composition Graph Options")
+class MOVIESCENECAPTURE_API UCompositionGraphCaptureSettings : public UObject
+{
+public:
+	GENERATED_BODY()
+	
+	/** A list of render passes to include in the capture. Leave empty to export all available passes. */
+	UPROPERTY(config, EditAnywhere, Category="Composition Graph Options")
+	FCompositionGraphCapturePasses IncludeRenderPasses;
+
+	/** Whether to capture the frames as HDR textures (*.exr format) */
+	UPROPERTY(config, EditAnywhere, Category="Composition Graph Options")
+	bool bCaptureFramesInHDR;
+
+	/** Custom post processing material to use for rendering */
+	UPROPERTY(config, EditAnywhere, Category="Composition Graph Options", meta=(AllowedClasses=""))
+	FStringAssetReference PostProcessingMaterial;
+};
+
+struct MOVIESCENECAPTURE_API FCompositionGraphCaptureProtocol : IMovieSceneCaptureProtocol
+{
+	/**~ IMovieSceneCaptureProtocol implementation */
+	virtual bool Initialize(const FCaptureProtocolInitSettings& InSettings, const ICaptureProtocolHost& Host);
+	virtual void CaptureFrame(const FFrameMetrics& FrameMetrics, const ICaptureProtocolHost& Host);
+	virtual void Tick() override;
+	virtual void Finalize() override;
+	virtual bool HasFinishedProcessing() const override;
+	/**~ End IMovieSceneCaptureProtocol implementation */
+
+private:
+	/** The viewport we are capturing from */
+	TWeakPtr<FSceneViewport> SceneViewport;
+
+	/** A view extension that we use to ensure we dump out the composition graph frames with the correct settings */
+	TSharedPtr<struct FSceneViewExtension, ESPMode::ThreadSafe> ViewExtension;
+
+	/** The render passes we want to export */
+	TArray<FString> RenderPasses;
+};
\ No newline at end of file
diff --git a/Engine/Source/Runtime/MovieSceneCapture/Public/Protocols/FrameGrabberProtocol.h b/Engine/Source/Runtime/MovieSceneCapture/Public/Protocols/FrameGrabberProtocol.h
new file mode 100644
--- /dev/null
+++ b/Engine/Source/Runtime/MovieSceneCapture/Public/Protocols/FrameGrabberProtocol.h
@@ -0,0 +1,61 @@
+// Copyright 1998-2015 Epic Games, Inc. All Rights Reserved.
+
+#pragma once
+
+#include "FrameGrabber.h"
+#include "IMovieSceneCaptureProtocol.h"
+
+#include "FrameGrabberProtocol.generated.h"
+
+UCLASS()
+class MOVIESCENECAPTURE_API UFrameGrabberProtocolSettings : public UObject
+{
+public:
+	UFrameGrabberProtocolSettings(const FObjectInitializer&) : DesiredPixelFormat(PF_B8G8R8A8), RingBufferSize(3) {}
+
+	GENERATED_BODY()
+
+	/** The pixel format we want to capture in */
+	EPixelFormat DesiredPixelFormat;
+
+	/** The size of the render-target resolution surface ring-buffer */
+	uint32 RingBufferSize;
+};
+
+struct MOVIESCENECAPTURE_API FFrameGrabberProtocol : IMovieSceneCaptureProtocol
+{
+protected:
+
+	/**~ IMovieSceneCaptureProtocol Implementation */
+	virtual bool HasFinishedProcessing() const;
+	virtual bool Initialize(const FCaptureProtocolInitSettings& InSettings, const ICaptureProtocolHost& Host);
+	virtual void CaptureFrame(const FFrameMetrics& FrameMetrics, const ICaptureProtocolHost& Host);
+	virtual void Tick();
+	virtual void Finalize();
+	/**~ End IMovieSceneCaptureProtocol Implementation */
+
+protected:
+
+	/**
+	 * Retrieve an arbitrary set of data that relates to the specified frame metrics.
+	 * This data will be passed through the capture pipeline, and will be accessible from ProcessFrame
+	 *
+	 * @param FrameMetrics	Metrics specific to the current frame
+	 * @param Host			The host that is managing this protocol
+
+	 * @return Shared pointer to a payload to associate with the frame, or nullptr
+	 */
+	virtual FFramePayloadPtr GetFramePayload(const FFrameMetrics& FrameMetrics, const ICaptureProtocolHost& Host) const = 0;
+
+	/**
+	 * Process a captured frame. This may be called on any thread.
+	 *
+	 * @param Frame			The captured frame data, including any payload retrieved from GetFramePayload
+	 */
+	virtual void ProcessFrame(FCapturedFrameData Frame) = 0;
+
+private:
+
+	/** The frame grabber, responsible for actually capturing frames */
+	TUniquePtr<FFrameGrabber> FrameGrabber;
+};
\ No newline at end of file
diff --git a/Engine/Source/Runtime/MovieSceneCapture/Public/Protocols/ImageSequenceProtocol.h b/Engine/Source/Runtime/MovieSceneCapture/Public/Protocols/ImageSequenceProtocol.h
new file mode 100644
--- /dev/null
+++ b/Engine/Source/Runtime/MovieSceneCapture/Public/Protocols/ImageSequenceProtocol.h
@@ -0,0 +1,89 @@
+// Copyright 1998-2015 Epic Games, Inc. All Rights Reserved.
+
+#pragma once
+
+#if WITH_EDITOR
+	#include "IImageWrapper.h"
+#endif
+
+#include "FrameGrabberProtocol.h"
+#include "ImageSequenceProtocol.generated.h"
+
+UCLASS(config=EditorPerProjectUserSettings, DisplayName="Image Encoding")
+class MOVIESCENECAPTURE_API UImageCaptureSettings : public UFrameGrabberProtocolSettings
+{
+public:
+	UImageCaptureSettings(const FObjectInitializer& Init) : UFrameGrabberProtocolSettings(Init), CompressionQuality(75) {}
+
+	GENERATED_BODY()
+
+	/** Level of compression to apply to the image, between 1 (worst quality, best compression) and 100 (best quality, worst compression)*/
+	UPROPERTY(config, EditAnywhere, Category=ImageSettings, meta=(ClampMin=1, ClampMax=100))
+	int32 CompressionQuality;
+};
+
+#if WITH_EDITOR
+
+/** Single runnable thread used to dispatch captured frames */
+struct FImageCaptureThread : public FRunnable
+{
+	FImageCaptureThread(EImageFormat::Type InFormat, int32 InCompressionQuality);
+	~FImageCaptureThread();
+
+	void Add(FCapturedFrameData Frame);
+	uint32 GetNumOutstandingFrames() const;
+	void Close();
+
+	virtual uint32 Run();
+	virtual void Stop() override;
+
+private:
+
+	void WriteFrameToDisk(FCapturedFrameData& Frame, IImageWrapperPtr ImageWrapper) const;
+
+private:
+	/** The thread itself */
+	FRunnableThread* Thread;
+	/** Command stack on which commands are pushed */
+	mutable FCriticalSection CommandMutex;
+	TArray<FCapturedFrameData> CapturedFrames;
+	/** Event that is triggered when we have something to do */
+	FEvent* WorkToDoEvent;
+	/** Event that is triggered when we've finished processing outstanding frames */
+	FEvent* ThreadEmptyEvent;
+	/** Set to false when the thread should terminate */
+	FThreadSafeBool bRunning;
+	/** The format we are writing out */
+	EImageFormat::Type Format;
+	/** Level of compression to apply to the image, between 1 (worst quality, best compression) and 100 (best quality, worst compression) */
+	int32 CompressionQuality;
+	/** array of image writer pointers for async writing */
+	TArray<IImageWrapperPtr> ImageWrappers;
+};
+
+struct MOVIESCENECAPTURE_API FImageSequenceProtocol : FFrameGrabberProtocol
+{
+	FImageSequenceProtocol(EImageFormat::Type InFormat);
+
+	/** ~FFrameGrabberProtocol implementation */
+	virtual bool Initialize(const FCaptureProtocolInitSettings& InSettings, const ICaptureProtocolHost& Host) override;
+	virtual FFramePayloadPtr GetFramePayload(const FFrameMetrics& FrameMetrics, const ICaptureProtocolHost& Host) const;
+	virtual void ProcessFrame(FCapturedFrameData Frame);
+	virtual void AddFormatMappings(TMap<FString, FStringFormatArg>& FormatMappings) const override;
+	virtual void Finalize() override;
+	virtual bool HasFinishedProcessing() const override;
+	/** ~End FFrameGrabberProtocol implementation */
+
+private:
+
+	/** Custom string format arguments for filenames */
+	TMap<FString, FStringFormatArg> StringFormatMap;
+	/** Level of compression to apply to the image, between 1 (worst quality, best compression) and 100 (best quality, worst compression)*/
+	int32 CompressionQuality;
+	/** The format of the image to write out */
+	EImageFormat::Type Format;
+	/** Thread responsible for writing out frames to disk */
+	TUniquePtr<FImageCaptureThread> CaptureThread;
+};
+
+#endif
\ No newline at end of file
diff --git a/Engine/Source/Runtime/MovieSceneCapture/Public/Protocols/VideoCaptureProtocol.h b/Engine/Source/Runtime/MovieSceneCapture/Public/Protocols/VideoCaptureProtocol.h
new file mode 100644
--- /dev/null
+++ b/Engine/Source/Runtime/MovieSceneCapture/Public/Protocols/VideoCaptureProtocol.h
@@ -0,0 +1,36 @@
+// Copyright 1998-2015 Epic Games, Inc. All Rights Reserved.
+
+#pragma once
+
+#include "AVIWriter.h"
+#include "FrameGrabberProtocol.h"
+#include "VideoCaptureProtocol.generated.h"
+
+UCLASS(config=EditorPerProjectUserSettings, DisplayName="Video Encoding")
+class MOVIESCENECAPTURE_API UVideoCaptureSettings : public UFrameGrabberProtocolSettings
+{
+public:
+	UVideoCaptureSettings(const FObjectInitializer& Init) : UFrameGrabberProtocolSettings(Init), bUseCompression(true), CompressionQuality(75) {}
+
+	GENERATED_BODY()
+
+	UPROPERTY(config, EditAnywhere, Category=VideoSettings)
+	bool bUseCompression;
+
+	UPROPERTY(config, EditAnywhere, Category=VideoSettings, meta=(ClampMin=1, ClampMax=100, EditCondition=bUseCompression))
+	float CompressionQuality;
+
+	UPROPERTY(config, EditAnywhere, Category=VideoSettings, AdvancedDisplay)
+	FString VideoCodec;
+};
+
+struct MOVIESCENECAPTURE_API FVideoCaptureProtocol : FFrameGrabberProtocol
+{
+	virtual bool Initialize(const FCaptureProtocolInitSettings& InSettings, const ICaptureProtocolHost& Host) override;
+	virtual void Finalize() override;
+	virtual FFramePayloadPtr GetFramePayload(const FFrameMetrics& FrameMetrics, const ICaptureProtocolHost& Host) const;
+	virtual void ProcessFrame(FCapturedFrameData Frame);
+	
+private:
+	TUniquePtr<FAVIWriter> AVIWriter;
+};
\ No newline at end of file
diff --git a/Engine/Source/Runtime/MovieSceneTracks/MovieSceneTracks.Build.cs b/Engine/Source/Runtime/MovieSceneTracks/MovieSceneTracks.Build.cs
--- a/Engine/Source/Runtime/MovieSceneTracks/MovieSceneTracks.Build.cs
+++ b/Engine/Source/Runtime/MovieSceneTracks/MovieSceneTracks.Build.cs
@@ -25,7 +25,7 @@ public class MovieSceneTracks : ModuleRules
 
 		PrivateDependencyModuleNames.AddRange(
 			new string[] {
-				"SlateCore",
+				"SlateCore"
 			}
 		);
 	}
diff --git a/Engine/Source/Runtime/MovieSceneTracks/Private/MovieSceneHitProxy.cpp b/Engine/Source/Runtime/MovieSceneTracks/Private/MovieSceneHitProxy.cpp
new file mode 100644
--- /dev/null
+++ b/Engine/Source/Runtime/MovieSceneTracks/Private/MovieSceneHitProxy.cpp
@@ -0,0 +1,6 @@
+// Copyright 1998-2015 Epic Games, Inc. All Rights Reserved.
+
+#include "MovieSceneTracksPrivatePCH.h"
+#include "MovieSceneHitProxy.h"
+
+IMPLEMENT_HIT_PROXY(HMovieSceneKeyProxy, HHitProxy);
diff --git a/Engine/Source/Runtime/MovieSceneTracks/Private/TrackInstances/MovieScene3DTransformTrackInstance.cpp b/Engine/Source/Runtime/MovieSceneTracks/Private/TrackInstances/MovieScene3DTransformTrackInstance.cpp
--- a/Engine/Source/Runtime/MovieSceneTracks/Private/TrackInstances/MovieScene3DTransformTrackInstance.cpp
+++ b/Engine/Source/Runtime/MovieSceneTracks/Private/TrackInstances/MovieScene3DTransformTrackInstance.cpp
@@ -5,7 +5,6 @@
 #include "MovieScene3DTransformTrackInstance.h"
 #include "MovieScene3DTransformTrack.h"
 
-
 FMovieScene3DTransformTrackInstance::FMovieScene3DTransformTrackInstance( UMovieScene3DTransformTrack& InTransformTrack )
 {
 	TransformTrack = &InTransformTrack;
@@ -23,6 +22,10 @@ void FMovieScene3DTransformTrackInstance::SaveState(const TArray<UObject*>& Runt
 			{
 				InitTransformMap.Add(RuntimeObjects[ObjIndex], SceneComponent->GetRelativeTransform());
 			}
+			if (InitMobilityMap.Find(RuntimeObjects[ObjIndex]) == nullptr)
+			{
+				InitMobilityMap.Add(RuntimeObjects[ObjIndex], SceneComponent->Mobility);
+			}
 		}
 	}
 }
@@ -45,13 +48,44 @@ void FMovieScene3DTransformTrackInstance::RestoreState(const TArray<UObject*>& R
 			{
 				SceneComponent->SetRelativeTransform(*Transform);
 			}
+
+			EComponentMobility::Type* ComponentMobility = InitMobilityMap.Find(RuntimeObjects[ObjIndex]);
+			if (ComponentMobility != nullptr)
+			{
+				SceneComponent->SetMobility(*ComponentMobility);
+			}
 		}
 	}
 }
 
+void FMovieScene3DTransformTrackInstance::UpdateRuntimeMobility(const TArray<UObject*>& RuntimeObjects)
+{
+	for( int32 ObjIndex = 0; ObjIndex < RuntimeObjects.Num(); ++ObjIndex )
+	{
+		USceneComponent* SceneComponent = MovieSceneHelpers::SceneComponentFromRuntimeObject(RuntimeObjects[ObjIndex]);
+
+		if (SceneComponent != nullptr)
+		{
+			if (SceneComponent->Mobility != EComponentMobility::Movable)
+			{
+				if (InitMobilityMap.Find(RuntimeObjects[ObjIndex]) == nullptr)
+				{
+					InitMobilityMap.Add(RuntimeObjects[ObjIndex], SceneComponent->Mobility);
+				}
+
+				SceneComponent->SetMobility(EComponentMobility::Movable);
+			}
+		}
+	}
+}
 
 void FMovieScene3DTransformTrackInstance::Update( float Position, float LastPosition, const TArray<UObject*>& RuntimeObjects, class IMovieScenePlayer& Player, FMovieSceneSequenceInstance& SequenceInstance, EMovieSceneUpdatePass UpdatePass ) 
 {
+	if (UpdatePass == MSUP_PreUpdate)
+	{
+		UpdateRuntimeMobility(RuntimeObjects);
+	}
+
 	FVector Translation;
 	FRotator Rotation;
 	FVector Scale;
@@ -78,3 +112,8 @@ void FMovieScene3DTransformTrackInstance::Update( float Position, float LastPosi
 		}
 	}
 }
+ 
+void FMovieScene3DTransformTrackInstance::RefreshInstance( const TArray<UObject*>& RuntimeObjects, IMovieScenePlayer& Player, FMovieSceneSequenceInstance& SequenceInstance )
+{
+	UpdateRuntimeMobility(RuntimeObjects);
+}
diff --git a/Engine/Source/Runtime/MovieSceneTracks/Private/Tracks/MovieSceneSubTrack.cpp b/Engine/Source/Runtime/MovieSceneTracks/Private/Tracks/MovieSceneSubTrack.cpp
--- a/Engine/Source/Runtime/MovieSceneTracks/Private/Tracks/MovieSceneSubTrack.cpp
+++ b/Engine/Source/Runtime/MovieSceneTracks/Private/Tracks/MovieSceneSubTrack.cpp
@@ -140,5 +140,11 @@ bool UMovieSceneSubTrack::SupportsMultipleRows() const
 	return true;
 }
 
+#if WITH_EDITORONLY_DATA
+FText UMovieSceneSubTrack::GetDisplayName() const
+{
+	return LOCTEXT("TrackName", "Subscenes");
+}
+#endif
 
 #undef LOCTEXT_NAMESPACE
