Copying //UE4/Dev-Sequencer to //UE4/Dev-Main (Source: //UE4/Dev-Sequencer @ 3003857)

diff --git a/Engine/Source/Runtime/MovieScene/Private/MovieSceneFolder.cpp b/Engine/Source/Runtime/MovieScene/Private/MovieSceneFolder.cpp
--- a/Engine/Source/Runtime/MovieScene/Private/MovieSceneFolder.cpp
+++ b/Engine/Source/Runtime/MovieScene/Private/MovieSceneFolder.cpp
@@ -5,6 +5,9 @@
 
 UMovieSceneFolder::UMovieSceneFolder( const FObjectInitializer& ObjectInitializer )
 	: Super( ObjectInitializer )
+#if WITH_EDITORONLY_DATA
+	, FolderColor(FColor::White)
+#endif
 {
 }
 
diff --git a/Engine/Source/Runtime/MovieScene/Public/MovieScene.h b/Engine/Source/Runtime/MovieScene/Public/MovieScene.h
--- a/Engine/Source/Runtime/MovieScene/Public/MovieScene.h
+++ b/Engine/Source/Runtime/MovieScene/Public/MovieScene.h
@@ -84,6 +84,11 @@ class MOVIESCENE_API UMovieScene
 {
 	GENERATED_UCLASS_BODY()
 
+public:
+
+	/**~ UObject implementation */
+	virtual void Serialize( FArchive& Ar ) override;
+
 public:
 
 #if WITH_EDITOR
@@ -469,6 +474,12 @@ public:
 	*/
 	void SetFixedFrameInterval( float InFixedFrameInterval );
 
+	/**
+	 * Calculates a fixed frame time based on a current time, a fixed frame interval, and an internal epsilon to account
+	 * for floating point consistency.
+	 */ 
+	static float CalculateFixedFrameTime( float Time, float FixedFrameInterval );
+
 public:
 	
 #if WITH_EDITORONLY_DATA
@@ -495,6 +506,12 @@ protected:
 	 */
 	void RemoveBinding(const FGuid& Guid);
 
+#if WITH_EDITOR
+	/** Templated helper for optimizing lists of possessables and spawnables for cook */
+	template<typename T>
+	void OptimizeObjectArray(TArray<T>& ObjectArray);
+#endif
+	
 protected:
 
 	/** Called after this object has been deserialized */
@@ -508,6 +525,10 @@ protected:
 
 private:
 
+	// Small value added for fixed frame interval calculations to make up for consistency in
+	// floating point calculations.
+	static const float FixedFrameIntervalEpsilon;
+
 	/**
 	 * Data-only blueprints for all of the objects that we we're able to spawn.
 	 * These describe objects and actors that we may instantiate at runtime,
diff --git a/Engine/Source/Runtime/MovieScene/Public/MovieSceneBinding.h b/Engine/Source/Runtime/MovieScene/Public/MovieSceneBinding.h
--- a/Engine/Source/Runtime/MovieScene/Public/MovieSceneBinding.h
+++ b/Engine/Source/Runtime/MovieScene/Public/MovieSceneBinding.h
@@ -80,7 +80,7 @@ struct FMovieSceneBinding
 	const FString& GetName() const
 	{
 		return BindingName;
-	} 
+	}
 
 	/**
 	 * Adds a new track to this binding
@@ -105,6 +105,15 @@ struct FMovieSceneBinding
 		return Tracks;
 	}
 
+#if WITH_EDITOR
+	/**
+	 * Perform cook-time optimization on this object binding
+	 * @param bShouldRemoveObject 		(Out) Boolean that will set to true if this whole binding should be considered redundant.
+	 */
+	void PerformCookOptimization(bool& bShouldRemoveObject);
+
+#endif
+
 private:
 
 	/** Object binding guid for runtime objects */
diff --git a/Engine/Source/Runtime/MovieScene/Public/MovieSceneFolder.h b/Engine/Source/Runtime/MovieScene/Public/MovieSceneFolder.h
--- a/Engine/Source/Runtime/MovieScene/Public/MovieSceneFolder.h
+++ b/Engine/Source/Runtime/MovieScene/Public/MovieSceneFolder.h
@@ -47,6 +47,27 @@ class MOVIESCENE_API UMovieSceneFolder : public UObject
 
 	virtual void Serialize( FArchive& Archive );
 
+#if WITH_EDITORONLY_DATA
+	/**
+	 * Get this folder's color.
+	 *
+	 * @return The folder color.
+	 */
+	const FColor& GetFolderColor() const
+	{
+		return FolderColor;
+	}
+
+	/**
+	 * Set this folder's color.
+	 *
+	 * @param InFolderColor The folder color to set.
+	 */
+	void SetFolderColor(const FColor& InFolderColor)
+	{
+		FolderColor = InFolderColor;
+	}
+#endif
 
 private:
 	/** The name of this folder. */
@@ -65,6 +86,12 @@ private:
 	UPROPERTY()
 	TArray<FString> ChildObjectBindingStrings;
 
+#if WITH_EDITORONLY_DATA
+	/** This folder's color */
+	UPROPERTY(EditAnywhere, Category=General, DisplayName=Color)
+	FColor FolderColor;
+#endif
+
 	/** The guids for the object bindings contained by this folder. */
 	TArray<FGuid> ChildObjectBindings;
 };
diff --git a/Engine/Source/Runtime/MovieSceneCapture/MovieSceneCapture.Build.cs b/Engine/Source/Runtime/MovieSceneCapture/MovieSceneCapture.Build.cs
--- a/Engine/Source/Runtime/MovieSceneCapture/MovieSceneCapture.Build.cs
+++ b/Engine/Source/Runtime/MovieSceneCapture/MovieSceneCapture.Build.cs
@@ -26,19 +26,21 @@ public class MovieSceneCapture : ModuleRules
 
 		PrivateDependencyModuleNames.AddRange(
 			new string[] {
-				"Core",
+                "AssetRegistry",
+                "Core",
 				"CoreUObject",
 				"Engine",
 				"InputCore",
 				"Json",
 				"JsonUtilities",
 				"MovieScene",
-				"RenderCore",
+                "MovieSceneTracks",
+                "RenderCore",
 				"RHI",
 				"ShaderCore",
 				"Slate",
 				"SlateCore",
 			}
 		);
-	}
+    }
 }
diff --git a/Engine/Source/Runtime/MovieSceneCapture/Private/AutomatedLevelSequenceCapture.cpp b/Engine/Source/Runtime/MovieSceneCapture/Private/AutomatedLevelSequenceCapture.cpp
--- a/Engine/Source/Runtime/MovieSceneCapture/Private/AutomatedLevelSequenceCapture.cpp
+++ b/Engine/Source/Runtime/MovieSceneCapture/Private/AutomatedLevelSequenceCapture.cpp
@@ -7,6 +7,11 @@
 #include "ErrorCodes.h"
 #include "SceneViewport.h"
 #include "ActiveMovieSceneCaptures.h"
+#include "JsonObjectConverter.h"
+#include "LevelSequenceBurnIn.h"
+#include "Tracks/MovieSceneCinematicShotTrack.h"
+#include "Sections/MovieSceneCinematicShotSection.h"
+#include "MovieSceneCaptureHelpers.h"
 
 UAutomatedLevelSequenceCapture::UAutomatedLevelSequenceCapture(const FObjectInitializer& Init)
 	: Super(Init)
@@ -23,9 +28,12 @@ UAutomatedLevelSequenceCapture::UAutomatedLevelSequenceCapture(const FObjectInit
 	EndFrame = 1;
 	WarmUpFrameCount = 0;
 	DelayBeforeWarmUp = 0.0f;
+	bWriteEditDecisionList = true;
 
 	RemainingDelaySeconds = 0.0f;
 	RemainingWarmUpFrames = 0;
+
+	BurnInOptions = Init.CreateDefaultSubobject<ULevelSequenceBurnInOptions>(this, "BurnInOptions");
 #endif
 }
 
@@ -36,6 +44,14 @@ void UAutomatedLevelSequenceCapture::SetLevelSequenceAsset(FString AssetPath)
 	LevelSequenceAsset = MoveTemp(AssetPath);
 }
 
+void UAutomatedLevelSequenceCapture::AddFormatMappings(TMap<FString, FStringFormatArg>& OutFormatMappings, const FFrameMetrics& FrameMetrics) const
+{
+	OutFormatMappings.Add(TEXT("shot"), CachedState.CurrentShotName.ToString());
+
+	const int32 FrameNumber = FMath::RoundToInt(CachedState.CurrentShotLocalTime * CachedState.Settings.FrameRate);
+	OutFormatMappings.Add(TEXT("shot_frame"), FString::Printf(*FrameNumberFormat, FrameNumber));
+}
+
 void UAutomatedLevelSequenceCapture::Initialize(TSharedPtr<FSceneViewport> InViewport, int32 PIEInstance)
 {
 	// Apply command-line overrides from parent class first. This needs to be called before setting up the capture strategy with the desired frame rate.
@@ -122,6 +138,9 @@ void UAutomatedLevelSequenceCapture::Initialize(TSharedPtr<FSceneViewport> InVie
 
 	if (Actor)
 	{
+		Actor->BurnInOptions = BurnInOptions;
+		Actor->RefreshBurnIn();
+
 		// Make sure we're not playing yet (in case AutoPlay was called from BeginPlay)
 		if( Actor->SequencePlayer != nullptr && Actor->SequencePlayer->IsPlaying() )
 		{
@@ -239,6 +258,7 @@ void UAutomatedLevelSequenceCapture::Tick(float DeltaSeconds)
 	}
 	else if( CaptureState == ELevelSequenceCaptureState::ReadyToWarmUp )
 	{
+		Actor->SequencePlayer->SetSnapshotSettings(FLevelSequenceSnapshotSettings(Settings.ZeroPadFrameNumbers, Settings.FrameRate));
 		Actor->SequencePlayer->StartPlayingNextTick();
 		// Start warming up
 		CaptureState = ELevelSequenceCaptureState::WarmingUp;
@@ -251,23 +271,122 @@ void UAutomatedLevelSequenceCapture::Tick(float DeltaSeconds)
 	{
 		// Start capturing - this will capture the *next* update from sequencer
 		CaptureState = ELevelSequenceCaptureState::FinishedWarmUp;
+		UpdateFrameState();
 		StartCapture();
 	}
 
 	if( bCapturing && !Actor->SequencePlayer->IsPlaying() )
 	{
 		Actor->SequencePlayer->OnSequenceUpdated().Remove( OnPlayerUpdatedBinding );
+		ExportEDL();
 		FinalizeWhenReady();
 	}
 }
 
 void UAutomatedLevelSequenceCapture::SequenceUpdated(const ULevelSequencePlayer& Player, float CurrentTime, float PreviousTime)
 {
-	LastSequenceUpdateDelta = CurrentTime - PreviousTime;
 	if (bCapturing)
 	{
-		CaptureThisFrame(LastSequenceUpdateDelta);
+		UpdateFrameState();
+		CaptureThisFrame(CurrentTime - PreviousTime);
+	}
+}
+
+void UAutomatedLevelSequenceCapture::UpdateFrameState()
+{
+	ALevelSequenceActor* Actor = LevelSequenceActor.Get();
+
+	if (Actor && Actor->SequencePlayer)
+	{
+		Actor->SequencePlayer->TakeFrameSnapshot(CachedState);
+	}
+}
+
+void UAutomatedLevelSequenceCapture::LoadFromConfig()
+{
+	UMovieSceneCapture::LoadFromConfig();
+
+	BurnInOptions->LoadConfig();
+	BurnInOptions->ResetSettings();
+	if (BurnInOptions->Settings)
+	{
+		BurnInOptions->Settings->LoadConfig();
+	}
+}
+
+void UAutomatedLevelSequenceCapture::SaveToConfig()
+{
+	BurnInOptions->SaveConfig();
+	if (BurnInOptions->Settings)
+	{
+		BurnInOptions->Settings->SaveConfig();
+	}
+
+	UMovieSceneCapture::SaveToConfig();
+}
+
+void UAutomatedLevelSequenceCapture::SerializeAdditionalJson(FJsonObject& Object)
+{
+	TSharedRef<FJsonObject> OptionsContainer = MakeShareable(new FJsonObject);
+	if (FJsonObjectConverter::UStructToJsonObject(BurnInOptions->GetClass(), BurnInOptions, OptionsContainer, 0, 0))
+	{
+		Object.SetField(TEXT("BurnInOptions"), MakeShareable(new FJsonValueObject(OptionsContainer)));
+	}
+
+	if (BurnInOptions->Settings)
+	{
+		TSharedRef<FJsonObject> SettingsDataObject = MakeShareable(new FJsonObject);
+		if (FJsonObjectConverter::UStructToJsonObject(BurnInOptions->Settings->GetClass(), BurnInOptions->Settings, SettingsDataObject, 0, 0))
+		{
+			Object.SetField(TEXT("BurnInOptionsInitSettings"), MakeShareable(new FJsonValueObject(SettingsDataObject)));
+		}
+	}
+}
+
+void UAutomatedLevelSequenceCapture::DeserializeAdditionalJson(const FJsonObject& Object)
+{
+	TSharedPtr<FJsonValue> OptionsContainer = Object.TryGetField(TEXT("BurnInOptions"));
+	if (OptionsContainer.IsValid())
+	{
+		FJsonObjectConverter::JsonAttributesToUStruct(OptionsContainer->AsObject()->Values, BurnInOptions->GetClass(), BurnInOptions, 0, 0);
+	}
+
+	BurnInOptions->ResetSettings();
+	if (BurnInOptions->Settings)
+	{
+		TSharedPtr<FJsonValue> SettingsDataObject = Object.TryGetField(TEXT("BurnInOptionsInitSettings"));
+		if (SettingsDataObject.IsValid())
+		{
+			FJsonObjectConverter::JsonAttributesToUStruct(SettingsDataObject->AsObject()->Values, BurnInOptions->Settings->GetClass(), BurnInOptions->Settings, 0, 0);
+		}
 	}
 }
 
+void UAutomatedLevelSequenceCapture::ExportEDL()
+{
+	if (!bWriteEditDecisionList)
+	{
+		return;
+	}
+	
+	ALevelSequenceActor* Actor = LevelSequenceActor.Get();
+	ULevelSequencePlayer* Player = Actor ? Actor->SequencePlayer : nullptr;
+	ULevelSequence* LevelSequence = Player ? Player->GetLevelSequence() : nullptr;
+	UMovieScene* MovieScene = LevelSequence ? LevelSequence->GetMovieScene() : nullptr;
+
+	if (!MovieScene)
+	{
+		return;
+	}
+
+	UMovieSceneCinematicShotTrack* ShotTrack = MovieScene->FindMasterTrack<UMovieSceneCinematicShotTrack>();
+	if (!ShotTrack)
+	{
+		return;
+	}
+
+	MovieSceneCaptureHelpers::ExportEDL(MovieScene, Settings.FrameRate, Settings.OutputDirectory.Path);
+}
+
+
 #endif
\ No newline at end of file
diff --git a/Engine/Source/Runtime/MovieSceneCapture/Private/CompositionGraphCaptureProtocol.cpp b/Engine/Source/Runtime/MovieSceneCapture/Private/CompositionGraphCaptureProtocol.cpp
--- a/Engine/Source/Runtime/MovieSceneCapture/Private/CompositionGraphCaptureProtocol.cpp
+++ b/Engine/Source/Runtime/MovieSceneCapture/Private/CompositionGraphCaptureProtocol.cpp
@@ -10,14 +10,15 @@
 
 struct FSceneViewExtension : ISceneViewExtension
 {
-	FSceneViewExtension(const TArray<FString>& InRenderPasses, bool bInCaptureFramesInHDR, UMaterialInterface* InPostProcessingMaterial)
-	: RenderPasses(InRenderPasses), RestoreDumpHDR(0), bNeedsCapture(true)
+	FSceneViewExtension(const TArray<FString>& InRenderPasses, bool bInCaptureFramesInHDR, int32 InCaptureGamut, UMaterialInterface* InPostProcessingMaterial)
+		: RenderPasses(InRenderPasses), bNeedsCapture(true), CaptureGamut(InCaptureGamut), RestoreDumpHDR(0), RestoreDumpGamut(HCGM_Rec709)
 	{
 		PostProcessingMaterial = InPostProcessingMaterial;
 		bCaptureFramesInHDR = bInCaptureFramesInHDR;
 
 		CVarDumpFrames = IConsoleManager::Get().FindConsoleVariable(TEXT("r.BufferVisualizationDumpFrames"));
 		CVarDumpFramesAsHDR = IConsoleManager::Get().FindConsoleVariable(TEXT("r.BufferVisualizationDumpFramesAsHDR"));
+		CVarDumpGamut = IConsoleManager::Get().FindConsoleVariable(TEXT("r.TonemapperOutputGamut"));
 
 		Disable();
 	}
@@ -38,7 +39,10 @@ struct FSceneViewExtension : ISceneViewExtension
 
 		bNeedsCapture = true;
 		RestoreDumpHDR = CVarDumpFramesAsHDR->GetInt();
+		RestoreDumpGamut = CVarDumpGamut->GetInt();
+
 		CVarDumpFramesAsHDR->Set(bCaptureFramesInHDR);
+		CVarDumpGamut->Set(CaptureGamut);
 		CVarDumpFrames->Set(1);
 	}
 
@@ -49,13 +53,10 @@ struct FSceneViewExtension : ISceneViewExtension
 			bNeedsCapture = false;
 			if (bFinalize)
 			{
-				CVarDumpFramesAsHDR->Set(0);
 				RestoreDumpHDR = 0;
 			}
-			else
-			{
-				CVarDumpFramesAsHDR->Set(RestoreDumpHDR);
-			}
+			CVarDumpFramesAsHDR->Set(RestoreDumpHDR);
+			CVarDumpGamut->Set(RestoreDumpGamut);
 			CVarDumpFrames->Set(0);
 		}
 	}
@@ -95,7 +96,7 @@ struct FSceneViewExtension : ISceneViewExtension
 			FWeightedBlendable Blendable(1.f, PostProcessingMaterial);
 			PostProcessingMaterial->OverrideBlendableSettings(InView, 1.f);
 		}
-		
+
 
 		// Ensure we're rendering at full size
 		InView.ViewRect = InView.UnscaledViewRect;
@@ -115,12 +116,16 @@ private:
 
 	UMaterialInterface* PostProcessingMaterial;
 
-	int32 RestoreDumpHDR;
 	bool bNeedsCapture;
+	int32 CaptureGamut;
 	FString OutputFilename;
 
 	IConsoleVariable* CVarDumpFrames;
 	IConsoleVariable* CVarDumpFramesAsHDR;
+	IConsoleVariable* CVarDumpGamut;
+
+	int32 RestoreDumpHDR;
+	int32 RestoreDumpGamut;
 };
 
 void UCompositionGraphCaptureSettings::OnReleaseConfig(FMovieSceneCaptureSettings& InSettings)
@@ -155,7 +160,7 @@ void UCompositionGraphCaptureSettings::OnLoadConfig(FMovieSceneCaptureSettings&
 			OutputFormat.InsertAt(FramePosition, TEXT("{material}"));
 		}
 		else
-		{ 
+		{
 			OutputFormat.Append(TEXT("{material}"));
 		}
 
@@ -170,6 +175,7 @@ bool FCompositionGraphCaptureProtocol::Initialize(const FCaptureProtocolInitSett
 	SceneViewport = InSettings.SceneViewport;
 
 	bool bCaptureFramesInHDR = false;
+	int32 CaptureGamut = HCGM_Rec709;
 
 	UMaterialInterface* PostProcessingMaterial = nullptr;
 	UCompositionGraphCaptureSettings* ProtocolSettings = CastChecked<UCompositionGraphCaptureSettings>(InSettings.ProtocolSettings);
@@ -177,29 +183,36 @@ bool FCompositionGraphCaptureProtocol::Initialize(const FCaptureProtocolInitSett
 	{
 		RenderPasses = ProtocolSettings->IncludeRenderPasses.Value;
 		bCaptureFramesInHDR = ProtocolSettings->bCaptureFramesInHDR;
+		CaptureGamut = ProtocolSettings->CaptureGamut;
 		PostProcessingMaterial = Cast<UMaterialInterface>(ProtocolSettings->PostProcessingMaterial.TryLoad());
 
 		FString OverrideRenderPasses;
-		if( FParse::Value( FCommandLine::Get(), TEXT( "-CustomRenderPasses=" ), OverrideRenderPasses ) )
+		if (FParse::Value(FCommandLine::Get(), TEXT("-CustomRenderPasses="), OverrideRenderPasses))
 		{
 			OverrideRenderPasses.ParseIntoArray(RenderPasses, TEXT(","), true);
 		}
 
 		bool bOverrideCaptureFramesInHDR;
-		if( FParse::Bool( FCommandLine::Get(), TEXT( "-CaptureFramesInHDR=" ), bOverrideCaptureFramesInHDR ) )
+		if (FParse::Bool(FCommandLine::Get(), TEXT("-CaptureFramesInHDR="), bOverrideCaptureFramesInHDR))
 		{
 			bCaptureFramesInHDR = bOverrideCaptureFramesInHDR;
 		}
+
+		int32 OverrideCaptureGamut;
+		if (FParse::Value(FCommandLine::Get(), TEXT("-CaptureGamut="), OverrideCaptureGamut))
+		{
+			CaptureGamut = OverrideCaptureGamut;
+		}
 	}
 
-	ViewExtension = MakeShareable(new FSceneViewExtension(RenderPasses, bCaptureFramesInHDR, PostProcessingMaterial));
+	ViewExtension = MakeShareable(new FSceneViewExtension(RenderPasses, bCaptureFramesInHDR, CaptureGamut, PostProcessingMaterial));
 
 	return true;
 }
 
 void FCompositionGraphCaptureProtocol::Finalize()
 {
-	ViewExtension->Disable();
+	ViewExtension->Disable(true);
 	GEngine->ViewExtensions.Remove(ViewExtension);
 }
 
diff --git a/Engine/Source/Runtime/MovieSceneCapture/Private/ImageSequenceProtocol.cpp b/Engine/Source/Runtime/MovieSceneCapture/Private/ImageSequenceProtocol.cpp
--- a/Engine/Source/Runtime/MovieSceneCapture/Private/ImageSequenceProtocol.cpp
+++ b/Engine/Source/Runtime/MovieSceneCapture/Private/ImageSequenceProtocol.cpp
@@ -104,7 +104,7 @@ void FImageSequenceProtocol::Finalize()
 	FFrameGrabberProtocol::Finalize();
 }
 
-FFramePayloadPtr FImageSequenceProtocol::GetFramePayload(const FFrameMetrics& FrameMetrics, const ICaptureProtocolHost& Host) const
+FFramePayloadPtr FImageSequenceProtocol::GetFramePayload(const FFrameMetrics& FrameMetrics, const ICaptureProtocolHost& Host)
 {
 	TSharedRef<FImageFrameData, ESPMode::ThreadSafe> FrameData = MakeShareable(new FImageFrameData);
 
@@ -117,6 +117,7 @@ FFramePayloadPtr FImageSequenceProtocol::GetFramePayload(const FFrameMetrics& Fr
 	}
 
 	FrameData->Filename = Host.GenerateFilename(FrameMetrics, Extension);
+	Host.EnsureFileWritable(FrameData->Filename);
 
 	// Add our custom formatting rules as well
 	// @todo: document these on the tooltip?
diff --git a/Engine/Source/Runtime/MovieSceneCapture/Private/MovieSceneCapture.cpp b/Engine/Source/Runtime/MovieSceneCapture/Private/MovieSceneCapture.cpp
--- a/Engine/Source/Runtime/MovieSceneCapture/Private/MovieSceneCapture.cpp
+++ b/Engine/Source/Runtime/MovieSceneCapture/Private/MovieSceneCapture.cpp
@@ -6,6 +6,8 @@
 #include "HighResScreenshot.h"
 #include "BufferVisualizationData.h"
 #include "SceneViewExtension.h"
+#include "JsonObjectConverter.h"
+#include "RemoteConfigIni.h"
 
 #if WITH_EDITOR
 #include "ImageWrapper.h"
@@ -36,6 +38,7 @@ FMovieSceneCaptureSettings::FMovieSceneCaptureSettings()
 	GameModeOverride = nullptr;
 	OutputFormat = TEXT("{world}");
 	FrameRate = 24;
+	ZeroPadFrameNumbers = 4;
 	bEnableTextureStreaming = false;
 	bCinematicMode = true;
 	bAllowMovement = false;
@@ -68,21 +71,29 @@ UMovieSceneCapture::UMovieSceneCapture(const FObjectInitializer& Initializer)
 
 void UMovieSceneCapture::PostInitProperties()
 {
-	if (!HasAnyFlags(RF_ClassDefaultObject))
+	InitializeSettings();
+
+	Super::PostInitProperties();
+}
+
+void UMovieSceneCapture::InitializeSettings()
+{
+	if (HasAnyFlags(RF_ClassDefaultObject))
 	{
-		if (ProtocolSettings)
-		{
-			ProtocolSettings->OnReleaseConfig(Settings);
-		}
-		ProtocolSettings = IMovieSceneCaptureModule::Get().GetProtocolRegistry().FactorySettingsType(CaptureType, this);
-		if (ProtocolSettings)
-		{
-			ProtocolSettings->LoadConfig();
-			ProtocolSettings->OnLoadConfig(Settings);
-		}
+		return;
 	}
 
-	Super::PostInitProperties();
+	if (ProtocolSettings)
+	{
+		ProtocolSettings->OnReleaseConfig(Settings);
+	}
+	
+	ProtocolSettings = IMovieSceneCaptureModule::Get().GetProtocolRegistry().FactorySettingsType(CaptureType, this);
+	if (ProtocolSettings)
+	{
+		ProtocolSettings->LoadConfig();
+		ProtocolSettings->OnLoadConfig(Settings);
+	}
 }
 
 void UMovieSceneCapture::Initialize(TSharedPtr<FSceneViewport> InSceneViewport, int32 PIEInstance)
@@ -125,6 +136,7 @@ void UMovieSceneCapture::Initialize(TSharedPtr<FSceneViewport> InSceneViewport,
 		if( FParse::Value( FCommandLine::Get(), TEXT( "-MovieFormat=" ), FormatOverride ) )
 		{
 			CaptureType = *FormatOverride;
+			InitializeSettings();
 		}
 
 		int32 FrameRateOverride;
@@ -148,6 +160,8 @@ void UMovieSceneCapture::Initialize(TSharedPtr<FSceneViewport> InSceneViewport,
 	FormatMappings.Add(TEXT("height"), FString::Printf(TEXT("%d"), CachedMetrics.Height));
 	FormatMappings.Add(TEXT("world"), GWorld->GetName());
 
+	FrameNumberFormat = FString::Printf(TEXT("%%0%dd"), Settings.ZeroPadFrameNumbers);
+
 	if (!HasAnyFlags(RF_ClassDefaultObject))
 	{
 		FActiveMovieSceneCaptures::Get().Add(this);
@@ -246,52 +260,58 @@ void UMovieSceneCapture::Finalize()
 
 FString UMovieSceneCapture::ResolveFileFormat(const FString& Format, const FFrameMetrics& FrameMetrics) const
 {
-	FormatMappings.Add(TEXT("frame"), FString::Printf(TEXT("%04d"), Settings.bUseRelativeFrameNumbers ? FrameMetrics.FrameNumber : FrameMetrics.FrameNumber + FrameNumberOffset));
-	
+	TMap<FString, FStringFormatArg> AllArgs = FormatMappings;
+
+	AllArgs.Add(TEXT("frame"), FString::Printf(*FrameNumberFormat, Settings.bUseRelativeFrameNumbers ? FrameMetrics.FrameNumber : FrameMetrics.FrameNumber + FrameNumberOffset));
+
+	AddFormatMappings(AllArgs, FrameMetrics);
+
 	if (CaptureProtocol.IsValid())
 	{
-		CaptureProtocol->AddFormatMappings(FormatMappings);
+		CaptureProtocol->AddFormatMappings(AllArgs);
 	}
-	return FString::Format(*Format, FormatMappings);
+	return FString::Format(*Format, AllArgs);
 }
 
-FString UMovieSceneCapture::GenerateFilename(const FFrameMetrics& FrameMetrics, const TCHAR* Extension) const
+void UMovieSceneCapture::EnsureFileWritable(const FString& File) const
 {
-	FString OutputDirectory = ResolveFileFormat(Settings.OutputDirectory.Path, FrameMetrics);
+	FString Directory = FPaths::GetPath(File);
 
-	if (!IFileManager::Get().DirectoryExists(*OutputDirectory))
+	if (!IFileManager::Get().DirectoryExists(*Directory))
 	{
-		IFileManager::Get().MakeDirectory(*OutputDirectory);
+		IFileManager::Get().MakeDirectory(*Directory);
 	}
 
-	const FString BaseFilename = OutputDirectory / ResolveFileFormat(Settings.OutputFormat, FrameMetrics);
-
-	FString ThisTry = BaseFilename + Extension;
-
 	if (Settings.bOverwriteExisting)
 	{
 		// Try and delete it first
-		while (IFileManager::Get().FileSize(*ThisTry) != -1 && !FPlatformFileManager::Get().GetPlatformFile().DeleteFile(*ThisTry))
+		while (IFileManager::Get().FileSize(*File) != -1 && !FPlatformFileManager::Get().GetPlatformFile().DeleteFile(*File))
 		{
 			// popup a message box
-			FText MessageText = FText::Format(LOCTEXT("UnableToRemoveFile_Format", "The destination file '{0}' could not be deleted because it's in use by another application.\n\nPlease close this application before continuing."), FText::FromString(ThisTry));
+			FText MessageText = FText::Format(LOCTEXT("UnableToRemoveFile_Format", "The destination file '{0}' could not be deleted because it's in use by another application.\n\nPlease close this application before continuing."), FText::FromString(File));
 			FPlatformMisc::MessageBoxExt(EAppMsgType::Ok, *MessageText.ToString(), *LOCTEXT("UnableToRemoveFile", "Unable to remove file").ToString());
 		}
-		return ThisTry;
 	}
+}
 
-	if (IFileManager::Get().FileSize(*ThisTry) == -1)
+FString UMovieSceneCapture::GenerateFilename(const FFrameMetrics& FrameMetrics, const TCHAR* Extension) const
+{
+	const FString BaseFilename = ResolveFileFormat(Settings.OutputDirectory.Path, FrameMetrics) / ResolveFileFormat(Settings.OutputFormat, FrameMetrics);
+
+	FString ThisTry = BaseFilename + Extension;
+
+	if (CaptureProtocol->CanWriteToFile(*ThisTry, Settings.bOverwriteExisting))
 	{
 		return ThisTry;
 	}
 
-	int32 DuplicateIndex = 1;
+	int32 DuplicateIndex = 2;
 	for (;;)
 	{
 		ThisTry = BaseFilename + FString::Printf(TEXT("_(%d)"), DuplicateIndex) + Extension;
 
 		// If the file doesn't exist, we can use that, else, increment the index and try again
-		if (IFileManager::Get().FileSize(*ThisTry) == -1)
+		if (CaptureProtocol->CanWriteToFile(*ThisTry, Settings.bOverwriteExisting))
 		{
 			return ThisTry;
 		}
@@ -302,26 +322,96 @@ FString UMovieSceneCapture::GenerateFilename(const FFrameMetrics& FrameMetrics,
 	return ThisTry;
 }
 
-#if WITH_EDITOR
-void UMovieSceneCapture::PostEditChangeProperty( struct FPropertyChangedEvent& PropertyChangedEvent)
+void UMovieSceneCapture::LoadFromConfig()
 {
-	FName PropertyName = (PropertyChangedEvent.MemberProperty != NULL) ? PropertyChangedEvent.MemberProperty->GetFName() : NAME_None;
+	LoadConfig();
+	ProtocolSettings->LoadConfig();
 
-	if (PropertyName == GET_MEMBER_NAME_CHECKED(UMovieSceneCapture, CaptureType))
+	FString JsonString;
+	FString Section = GetClass()->GetPathName() + TEXT("_Json");
+
+	if (GConfig->GetString( *Section, TEXT("Data"), JsonString, GEditorSettingsIni))
 	{
-		if (ProtocolSettings)
+		FString UnescapedString = FRemoteConfig::ReplaceIniSpecialCharWithChar(JsonString).ReplaceEscapedCharWithChar();
+
+		TSharedPtr<FJsonObject> RootObject;
+		TSharedRef<TJsonReader<> > JsonReader = TJsonReaderFactory<>::Create(UnescapedString);
+		if (FJsonSerializer::Deserialize(JsonReader, RootObject) && RootObject.IsValid())
 		{
-			ProtocolSettings->OnReleaseConfig(Settings);
+			DeserializeAdditionalJson(*RootObject);
 		}
+	}
+}
+
+void UMovieSceneCapture::SaveToConfig()
+{
+	TSharedRef<FJsonObject> Json = MakeShareable(new FJsonObject);
+	SerializeAdditionalJson(*Json);
+
+	FString JsonString;
+	TSharedRef<TJsonWriter<> > JsonWriter = TJsonWriterFactory<>::Create(&JsonString, 0);
+	if (FJsonSerializer::Serialize( Json, JsonWriter ))
+	{
+		FString Section = GetClass()->GetPathName() + TEXT("_Json");
 
-		ProtocolSettings = IMovieSceneCaptureModule::Get().GetProtocolRegistry().FactorySettingsType(CaptureType, this);
-		if (ProtocolSettings)
+		FString EscapedJson = FRemoteConfig::ReplaceIniCharWithSpecialChar(JsonString).ReplaceCharWithEscapedChar();
+
+		GConfig->SetString( *Section, TEXT("Data"), *EscapedJson, GEditorSettingsIni);
+		GConfig->Flush(false, GEditorSettingsIni);
+	}
+
+	SaveConfig();
+	ProtocolSettings->SaveConfig();
+}
+
+void UMovieSceneCapture::SerializeJson(FJsonObject& Object)
+{
+	if (ProtocolSettings)
+	{
+		Object.SetField(TEXT("ProtocolType"), MakeShareable(new FJsonValueString(ProtocolSettings->GetClass()->GetPathName())));
+		TSharedRef<FJsonObject> ProtocolDataObject = MakeShareable(new FJsonObject);
+		if (FJsonObjectConverter::UStructToJsonObject(ProtocolSettings->GetClass(), ProtocolSettings, ProtocolDataObject, 0, 0))
 		{
-			ProtocolSettings->LoadConfig();
-			ProtocolSettings->OnLoadConfig(Settings);
+			Object.SetField(TEXT("ProtocolData"), MakeShareable(new FJsonValueObject(ProtocolDataObject)));
 		}
 	}
 
+	SerializeAdditionalJson(Object);
+}
+
+void UMovieSceneCapture::DeserializeJson(const FJsonObject& Object)
+{
+	TSharedPtr<FJsonValue> ProtocolTypeField = Object.TryGetField(TEXT("ProtocolType"));
+	if (ProtocolTypeField.IsValid())
+	{
+		UClass* ProtocolTypeClass = FindObject<UClass>(nullptr, *ProtocolTypeField->AsString());
+		if (ProtocolTypeClass)
+		{
+			ProtocolSettings = NewObject<UMovieSceneCaptureProtocolSettings>(this, ProtocolTypeClass);
+			if (ProtocolSettings)
+			{
+				TSharedPtr<FJsonValue> ProtocolDataField = Object.TryGetField(TEXT("ProtocolData"));
+				if (ProtocolDataField.IsValid())
+				{
+					FJsonObjectConverter::JsonAttributesToUStruct(ProtocolDataField->AsObject()->Values, ProtocolTypeClass, ProtocolSettings, 0, 0);
+				}
+			}
+		}
+	}
+
+	DeserializeAdditionalJson(Object);
+}
+
+#if WITH_EDITOR
+void UMovieSceneCapture::PostEditChangeProperty( struct FPropertyChangedEvent& PropertyChangedEvent)
+{
+	FName PropertyName = (PropertyChangedEvent.MemberProperty != NULL) ? PropertyChangedEvent.MemberProperty->GetFName() : NAME_None;
+
+	if (PropertyName == GET_MEMBER_NAME_CHECKED(UMovieSceneCapture, CaptureType))
+	{
+		InitializeSettings();
+	}
+
 	Super::PostEditChangeProperty(PropertyChangedEvent);
 }
 #endif
diff --git a/Engine/Source/Runtime/MovieSceneCapture/Private/MovieSceneCaptureHelpers.cpp b/Engine/Source/Runtime/MovieSceneCapture/Private/MovieSceneCaptureHelpers.cpp
new file mode 100644
--- /dev/null
+++ b/Engine/Source/Runtime/MovieSceneCapture/Private/MovieSceneCaptureHelpers.cpp
@@ -0,0 +1,461 @@
+// Copyright 1998-2016 Epic Games, Inc. All Rights Reserved.
+
+#include "MovieSceneCapturePCH.h"
+#include "MovieSceneCaptureHelpers.h"
+#include "MovieScene.h"
+#include "MovieSceneSection.h"
+#include "MovieSceneAudioTrack.h"
+#include "MovieSceneCinematicShotSection.h"
+#include "MovieSceneCinematicShotTrack.h"
+#include "LevelSequence.h"
+#include "AssetRegistryModule.h"
+#include "CoreMisc.h"
+
+/* MovieSceneCaptureHelpers
+ *****************************************************************************/
+
+struct FShotData
+{
+	enum class ETrackType
+	{
+		TT_Video,
+		TT_A,
+		TT_A2,
+		TT_AA,
+		TT_None
+	};
+
+	enum class EEditType
+	{
+		ET_Cut,
+		ET_Dissolve,
+		ET_Wipe,
+		ET_KeyEdit,
+		ET_None
+	};
+
+	FShotData(const FString& InElementName, const FString& InElementPath, ETrackType InTrackType, EEditType InEditType, float InSourceInTime, float InSourceOutTime, float InEditInTime, float InEditOutTime) : 
+		  ElementName(InElementName)
+		, ElementPath(InElementPath)
+		, TrackType(InTrackType)
+		, EditType(InEditType)
+		, SourceInTime(InSourceInTime)
+		, SourceOutTime(InSourceOutTime)
+		, EditInTime(InEditInTime)
+		, EditOutTime(InEditOutTime) {}
+
+	bool operator<(const FShotData& Other) const { return EditInTime < Other.EditInTime; }
+
+	FString ElementName;
+	FString ElementPath;
+	ETrackType TrackType;
+	EEditType EditType;
+	float SourceInTime;
+	float SourceOutTime;
+	float EditInTime;
+	float EditOutTime;
+};
+
+bool SMPTEToTime(FString SMPTE, float FrameRate, float& OutTime)
+{
+	TArray<FString> OutArray;
+	SMPTE.ParseIntoArray(OutArray, TEXT(":"));
+	
+	if (OutArray.Num() == 4)
+	{
+		float Hours = FCString::Atof(*OutArray[0]);
+		float Minutes = FCString::Atof(*OutArray[1]);
+		float Seconds = FCString::Atof(*OutArray[2]);
+		float Frames = FCString::Atof(*OutArray[3]);
+
+		OutTime = Hours * 3600.f + Minutes * 60.f  + Seconds + (Frames / FrameRate);
+		return true;
+	}
+	// The edl is in frames
+	else
+	{
+		float Frames = FCString::Atof(*SMPTE);
+		OutTime = Frames * FrameRate;
+	}
+
+	return false;
+}
+
+FString TimeToSMPTE(float InTime, float FrameRate)
+{
+	float Hours, Minutes, Seconds, Frames;
+	Frames = 0.5f + FrameRate * (InTime - FMath::FloorToInt(InTime));
+
+	Hours = FMath::FloorToInt(InTime / (60.f * 60.f));
+	InTime -= Hours * 60.f * 60.f;
+	Minutes = FMath::FloorToInt(InTime / 60.f);
+	InTime -= Minutes * 60.f;
+	Seconds = InTime;
+
+	FString OutputString = FString::Printf(TEXT("%02d:%02d:%02d:%02d"), (int)Hours, (int)Minutes, (int)Seconds, (int)Frames);
+
+	return OutputString;
+}
+
+void ParseFromEDL(const FString& InputString, float FrameRate, TArray<FShotData>& OutShotData)
+{
+	TArray<FString> InputLines;
+	InputString.ParseIntoArray(InputLines, TEXT("\n"));
+
+	bool bFoundEventLine = false;
+	FString EventName;
+	FString AuxName;
+	FString ReelName;
+	FShotData::ETrackType TrackType = FShotData::ETrackType::TT_None;
+	FShotData::EEditType EditType = FShotData::EEditType::ET_None;
+	float SourceInTime = 0.f;
+	float SourceOutTime = 0.f;
+	float EditInTime = 0.f;
+	float EditOutTime = 0.f;
+
+	for (FString InputLine : InputLines)
+	{
+		TArray<FString> InputChars;
+		InputLine.ParseIntoArrayWS(InputChars);
+
+		// First look for :
+		// 001 AX V C 00:00:00:00 00:00:12:02 00:00:07:20 00:00:12:03
+		if (!bFoundEventLine)
+		{
+			if (InputChars.Num() == 8)
+			{
+				EventName = InputChars[0];
+				AuxName = InputChars[1]; // Typically AX but unused in this case
+
+				if (InputChars[2] == TEXT("V"))
+				{
+					TrackType = FShotData::ETrackType::TT_Video;
+				}
+				else if (InputChars[2] == TEXT("A"))
+				{
+					TrackType = FShotData::ETrackType::TT_A;
+				}
+				else if (InputChars[2] == TEXT("A2"))
+				{
+					TrackType = FShotData::ETrackType::TT_A2;
+				}
+				else if (InputChars[2] == TEXT("AA"))
+				{
+					TrackType = FShotData::ETrackType::TT_AA;
+				}
+
+				if (InputChars[3] == TEXT("C"))
+				{
+					EditType = FShotData::EEditType::ET_Cut;
+				}
+				else if (InputChars[3] == TEXT("D"))
+				{
+					EditType = FShotData::EEditType::ET_Dissolve;
+				}
+				else if (InputChars[3] == TEXT("W"))
+				{
+					EditType = FShotData::EEditType::ET_Wipe;
+				}
+				else if (InputChars[3] == TEXT("K"))
+				{
+					EditType = FShotData::EEditType::ET_KeyEdit;
+				}
+
+				// If everything checks out
+				if (TrackType != FShotData::ETrackType::TT_None &&
+					EditType != FShotData::EEditType::ET_None &&
+					SMPTEToTime(InputChars[4], FrameRate, SourceInTime) &&
+					SMPTEToTime(InputChars[5], FrameRate, SourceOutTime) &&
+					SMPTEToTime(InputChars[6], FrameRate, EditInTime) &&
+					SMPTEToTime(InputChars[7], FrameRate, EditOutTime))
+				{
+					bFoundEventLine = true;
+					continue; // Go to the next line
+				}
+			}
+		}
+		
+		// Then look for:
+		// * FROM CLIP NAME: shot0010_001.avi
+		else
+		{
+			if (InputChars.Num() == 5 &&
+				InputChars[0] == TEXT("*") &&
+				InputChars[1].ToUpper() == TEXT("FROM") &&
+				InputChars[2].ToUpper() == TEXT("CLIP") &&
+				InputChars[3].ToUpper() == TEXT("NAME:"))
+			{
+				ReelName = InputChars[4];
+
+				//@todo can't assume avi's written out
+				ReelName = ReelName.LeftChop(4); // strip .avi
+
+				FString ElementName = ReelName;
+				FString ElementPath = ElementName;
+
+				// If everything checks out add to OutShotData
+				OutShotData.Add(FShotData(ElementName, ElementPath, TrackType, EditType, SourceInTime, SourceOutTime, EditInTime, EditOutTime));
+				bFoundEventLine = false; // Reset and go to next line to look for element line
+				continue;
+			}
+		}
+	}
+}
+
+void FormatForEDL(FString& OutputString, const FString& SequenceName, float FrameRate, const TArray<FShotData>& InShotData)
+{
+	OutputString += TEXT("TITLE: ") + SequenceName + TEXT("\n");
+	OutputString += TEXT("FCM: NON-DROP FRAME\n\n");
+
+	for (int32 EventIndex = 0; EventIndex < InShotData.Num(); ++EventIndex)
+	{
+		FString EventName = FString::Printf(TEXT("%03d"), EventIndex+1);
+
+		FString ReelName = InShotData[EventIndex].ElementName; 
+
+		FString TypeName;
+		if (InShotData[EventIndex].TrackType == FShotData::ETrackType::TT_Video)
+		{
+			TypeName = TEXT("V");
+		}
+		else if (InShotData[EventIndex].TrackType == FShotData::ETrackType::TT_A)
+		{
+			TypeName = TEXT("A");
+		}
+		else if (InShotData[EventIndex].TrackType == FShotData::ETrackType::TT_A2)
+		{
+			TypeName = TEXT("A2");
+		}
+		else if (InShotData[EventIndex].TrackType == FShotData::ETrackType::TT_AA)
+		{
+			TypeName = TEXT("AA");
+		}
+
+		FString EditName;
+		if (InShotData[EventIndex].EditType == FShotData::EEditType::ET_Cut)
+		{
+			EditName = TEXT("C");
+		}
+		else if (InShotData[EventIndex].EditType == FShotData::EEditType::ET_Dissolve)
+		{
+			EditName = TEXT("D");
+		}
+		else if (InShotData[EventIndex].EditType == FShotData::EEditType::ET_Wipe)
+		{
+			EditName = TEXT("W");
+		}
+		else if (InShotData[EventIndex].EditType == FShotData::EEditType::ET_KeyEdit)
+		{
+			EditName = TEXT("K");
+		}
+
+		FString SourceSMPTEIn = TimeToSMPTE(InShotData[EventIndex].SourceInTime, FrameRate);
+		FString SourceSMPTEOut = TimeToSMPTE(InShotData[EventIndex].SourceOutTime, FrameRate);
+		FString EditSMPTEIn = TimeToSMPTE(InShotData[EventIndex].EditInTime, FrameRate);
+		FString EditSMPTEOut = TimeToSMPTE(InShotData[EventIndex].EditOutTime, FrameRate);
+
+		OutputString += EventName + TEXT(" ") + TEXT("AX ") + TypeName + TEXT(" ") + EditName + TEXT(" ");
+		OutputString += SourceSMPTEIn + TEXT(" ") + SourceSMPTEOut + TEXT(" ") + EditSMPTEIn + TEXT(" ") + EditSMPTEOut + TEXT("\n");
+		OutputString += TEXT("* FROM CLIP NAME: ") + ReelName + TEXT("\n\n");
+	}
+}
+
+void FormatForRV(FString& OutputString, const FString& SequenceName, float FrameRate, const TArray<FShotData>& InShotData)
+{
+	// Header
+	OutputString += TEXT("GTOa (3)\n\n");
+	OutputString += TEXT("rv : RVSession (2)\n");
+	OutputString += TEXT("{\n");
+	OutputString += TEXT("\tsession\n");
+	OutputString += TEXT("\t{\n");
+	OutputString += TEXT("\t\tfloat fps = ") + FString::Printf(TEXT("%f"), FrameRate) + TEXT("\n");
+	OutputString += TEXT("\t\tint realtime = 1\n");
+	OutputString += TEXT("\t}\n\n");
+	OutputString += TEXT("\twriter\n");
+	OutputString += TEXT("\t{\n");
+	OutputString += TEXT("\t\tstring name = \"rvSession.py\"\n");
+	OutputString += TEXT("\t\tstring version = \"0.3\"\n");
+	OutputString += TEXT("\t}\n");
+	OutputString += TEXT("}\n\n");
+
+	// Body
+	for (int32 EventIndex = 0; EventIndex < InShotData.Num(); ++EventIndex)
+	{
+		FString SourceName = FString::Printf(TEXT("sourceGroup%06d"), EventIndex);
+
+		int32 SourceInFrame = FMath::RoundToInt(InShotData[EventIndex].SourceInTime * FrameRate);
+		int32 SourceOutFrame = FMath::RoundToInt(InShotData[EventIndex].SourceOutTime * FrameRate);
+
+		OutputString += SourceName + TEXT(" : RVSourceGroup (1)\n");
+		OutputString += TEXT("{\n");
+		OutputString += TEXT("\tui\n");
+		OutputString += TEXT("\t{\n");
+		OutputString += TEXT("\t\tstring name = \"") + InShotData[EventIndex].ElementName + TEXT("\"\n");
+		OutputString += TEXT("\t}\n");
+		OutputString += TEXT("}\n\n");
+
+		OutputString += SourceName + TEXT("_source : RVFileSource (1)\n");
+		OutputString += TEXT("{\n");
+		OutputString += TEXT("\tcut\n");
+		OutputString += TEXT("\t{\n");
+		OutputString += TEXT("\t\tint in = ") + FString::FromInt(SourceInFrame) + TEXT("\n");
+		OutputString += TEXT("\t\tint out = ") + FString::FromInt(SourceOutFrame) + TEXT("\n");
+		OutputString += TEXT("\t}\n\n");
+
+		OutputString += TEXT("\tgroup\n");
+		OutputString += TEXT("\t{\n");
+		OutputString += TEXT("\t\tint noMovieAudio = 1\n");
+		OutputString += TEXT("\t}\n\n");
+
+		OutputString += TEXT("\tmedia\n");
+		OutputString += TEXT("\t{\n");
+		OutputString += TEXT("\t\tstring movie = \"") + InShotData[EventIndex].ElementPath + TEXT("\"\n");
+		OutputString += TEXT("\t\tstring shot = \"\"\n");
+		OutputString += TEXT("\t}\n");
+		OutputString += TEXT("}\n\n");
+	}
+}
+
+bool MovieSceneCaptureHelpers::ImportEDL(UMovieScene* InMovieScene, float InFrameRate, FString InFilename)
+{
+	FString InputString;
+	if (!FFileHelper::LoadFileToString(InputString, *InFilename))
+	{
+		return false;
+	}
+
+	TArray<FShotData> ShotDataArray;
+	ParseFromEDL(InputString, InFrameRate, ShotDataArray);
+
+	UMovieSceneCinematicShotTrack* CinematicShotTrack = Cast<UMovieSceneCinematicShotTrack>(InMovieScene->FindMasterTrack<UMovieSceneCinematicShotTrack>());
+	if (!CinematicShotTrack)
+	{
+		CinematicShotTrack = InMovieScene->AddMasterTrack<UMovieSceneCinematicShotTrack>();
+	}
+
+	for (FShotData ShotData : ShotDataArray)
+	{
+		if (ShotData.TrackType == FShotData::ETrackType::TT_Video)
+		{
+			UMovieSceneCinematicShotSection* ShotSection = nullptr;
+
+			FString ShotName = ShotData.ElementName;
+			for (UMovieSceneSection* Section : CinematicShotTrack->GetAllSections())
+			{
+				UMovieSceneCinematicShotSection* CinematicShotSection = Cast<UMovieSceneCinematicShotSection>(Section);
+				if (CinematicShotSection->GetSequence()->GetName() == ShotName)
+				{
+					ShotSection = CinematicShotSection;
+					break;
+				}
+			}
+
+			// If the shot doesn't already exist, create it
+			if (!ShotSection)
+			{
+				FAssetRegistryModule& AssetRegistryModule = FModuleManager::LoadModuleChecked<FAssetRegistryModule>(TEXT("AssetRegistry"));
+
+				// Collect a full list of assets with the specified class
+				TArray<FAssetData> AssetDataArray;
+				AssetRegistryModule.Get().GetAssetsByClass(ULevelSequence::StaticClass()->GetFName(), AssetDataArray);
+
+				for (FAssetData AssetData : AssetDataArray)
+				{
+					if (AssetData.AssetName == *ShotName)
+					{	
+						CinematicShotTrack->Modify();
+						ShotSection = Cast<UMovieSceneCinematicShotSection>(CinematicShotTrack->AddSequence(Cast<ULevelSequence>(AssetData.GetAsset()), ShotData.EditInTime, ShotData.EditOutTime - ShotData.EditInTime));
+						break;
+					}
+				}
+			}
+
+			// Conform this shot section
+			if (ShotSection)
+			{
+				ShotSection->Modify();
+				ShotSection->StartOffset = ShotData.SourceInTime;
+				ShotSection->SetStartTime(ShotData.EditInTime);
+				ShotSection->SetEndTime(ShotData.EditOutTime);
+			}
+		}
+	}
+
+	return true;
+}
+
+bool MovieSceneCaptureHelpers::ExportEDL(const UMovieScene* InMovieScene, float InFrameRate, FString InSaveDirectory)
+{
+	FString SequenceName = InMovieScene->GetOuter()->GetName();
+	FString SequencePath = InMovieScene->GetOuter()->GetPathName();
+	SequencePath = FPaths::GetPath(SequencePath);
+
+	// Pop open a dialog to request the location of the edl (.edl, or .rv formats accepted)
+	TArray<FString> SaveFilenames;
+	if (!InSaveDirectory.IsEmpty())
+	{
+		SaveFilenames.Add(InSaveDirectory / SequenceName + TEXT(".rv"));
+		SaveFilenames.Add(InSaveDirectory / SequenceName + TEXT(".edl"));
+	}
+
+	TArray<FShotData> ShotDataArray;
+
+	for (auto MasterTrack : InMovieScene->GetMasterTracks())
+	{
+		TRange<float> PlaybackRange = InMovieScene->GetPlaybackRange();
+		if (MasterTrack->IsA(UMovieSceneCinematicShotTrack::StaticClass()))
+		{
+			UMovieSceneCinematicShotTrack* CinematicShotTrack = Cast<UMovieSceneCinematicShotTrack>(MasterTrack);
+
+			for (auto ShotSection : CinematicShotTrack->GetAllSections())
+			{
+				UMovieSceneCinematicShotSection* CinematicShotSection = Cast<UMovieSceneCinematicShotSection>(ShotSection);
+				
+				FString ShotName = CinematicShotSection->GetShotDisplayName().ToString();
+				FString ShotPath = CinematicShotSection->GetSequence()->GetMovieScene()->GetOuter()->GetPathName();
+
+				float SourceInTime = CinematicShotSection->StartOffset;
+				float SourceOutTime = CinematicShotSection->StartOffset + CinematicShotSection->GetTimeSize();
+				float EditInTime = CinematicShotSection->GetStartTime();
+				float EditOutTime = CinematicShotSection->GetEndTime();
+
+				EditInTime = FMath::Clamp(EditInTime, PlaybackRange.GetLowerBoundValue(), PlaybackRange.GetUpperBoundValue());
+				EditOutTime = FMath::Clamp(EditOutTime, PlaybackRange.GetLowerBoundValue(), PlaybackRange.GetUpperBoundValue());
+
+				//@todo can't assume avi's written out
+				ShotName += ".avi";
+
+				//@todo shotpath should really be moviefile path
+				ShotPath = ShotName;
+
+				ShotDataArray.Add(FShotData(ShotName, ShotPath, FShotData::ETrackType::TT_Video, FShotData::EEditType::ET_Cut, SourceInTime, SourceOutTime, EditInTime, EditOutTime));
+			}
+		}
+		else if (MasterTrack->IsA(UMovieSceneAudioTrack::StaticClass()))
+		{
+			UMovieSceneAudioTrack* AudioTrack = Cast<UMovieSceneAudioTrack>(MasterTrack);
+
+			//@todo support audio clips
+		}
+	}
+
+	ShotDataArray.Sort();
+
+	for (auto SaveFilename : SaveFilenames)
+	{
+		FString OutputString;
+
+		if (FPaths::GetExtension(SaveFilename).ToUpper() == TEXT("EDL"))
+		{
+			FormatForEDL(OutputString, SequenceName, InFrameRate, ShotDataArray);
+		}
+		else if (FPaths::GetExtension(SaveFilename).ToUpper() == TEXT("RV"))
+		{
+			FormatForRV(OutputString, SequenceName, InFrameRate, ShotDataArray);
+		}
+
+		FFileHelper::SaveStringToFile(OutputString, *SaveFilename);
+	}
+
+	return true;
+}
\ No newline at end of file
diff --git a/Engine/Source/Runtime/MovieSceneCapture/Private/MovieSceneCaptureModule.cpp b/Engine/Source/Runtime/MovieSceneCapture/Private/MovieSceneCaptureModule.cpp
--- a/Engine/Source/Runtime/MovieSceneCapture/Private/MovieSceneCaptureModule.cpp
+++ b/Engine/Source/Runtime/MovieSceneCapture/Private/MovieSceneCaptureModule.cpp
@@ -162,23 +162,10 @@ private:
 						return nullptr;
 					}
 
-					// Now deserialize the protocol data
-					auto ProtocolTypeField = RootObject->TryGetField(TEXT("ProtocolType"));
-					if (ProtocolTypeField.IsValid())
+					TSharedPtr<FJsonValue> AdditionalDataField = RootObject->TryGetField(TEXT("AdditionalData"));
+					if (AdditionalDataField.IsValid())
 					{
-						UClass* ProtocolTypeClass = FindObject<UClass>(nullptr, *ProtocolTypeField->AsString());
-						if (ProtocolTypeClass)
-						{
-							Capture->ProtocolSettings = NewObject<UMovieSceneCaptureProtocolSettings>(Capture, ProtocolTypeClass);
-							if (Capture->ProtocolSettings)
-							{
-								auto ProtocolDataField = RootObject->TryGetField(TEXT("ProtocolData"));
-								if (ProtocolDataField.IsValid())
-								{
-									FJsonObjectConverter::JsonAttributesToUStruct(ProtocolDataField->AsObject()->Values, ProtocolTypeClass, Capture->ProtocolSettings, 0, 0);
-								}
-							}
-						}
+						Capture->DeserializeJson(*AdditionalDataField->AsObject());
 					}
 				}
 			}
diff --git a/Engine/Source/Runtime/MovieSceneCapture/Private/MovieSceneCaptureProtocol.cpp b/Engine/Source/Runtime/MovieSceneCapture/Private/MovieSceneCaptureProtocol.cpp
--- a/Engine/Source/Runtime/MovieSceneCapture/Private/MovieSceneCaptureProtocol.cpp
+++ b/Engine/Source/Runtime/MovieSceneCapture/Private/MovieSceneCaptureProtocol.cpp
@@ -3,6 +3,13 @@
 #include "MovieSceneCapturePCH.h"
 #include "IMovieSceneCaptureProtocol.h"
 #include "SceneViewport.h"
+#include "FileManager.h"
+
+
+bool IMovieSceneCaptureProtocol::CanWriteToFile(const TCHAR* InFilename, bool bOverwriteExisting) const
+{
+	return bOverwriteExisting || IFileManager::Get().FileSize(InFilename) == -1;
+}
 
 FCaptureProtocolInitSettings FCaptureProtocolInitSettings::FromSlateViewport(TSharedRef<FSceneViewport> InSceneViewport, UObject* InProtocolSettings)
 {
diff --git a/Engine/Source/Runtime/MovieSceneCapture/Private/VideoCaptureProtocol.cpp b/Engine/Source/Runtime/MovieSceneCapture/Private/VideoCaptureProtocol.cpp
--- a/Engine/Source/Runtime/MovieSceneCapture/Private/VideoCaptureProtocol.cpp
+++ b/Engine/Source/Runtime/MovieSceneCapture/Private/VideoCaptureProtocol.cpp
@@ -6,26 +6,43 @@
 
 bool FVideoCaptureProtocol::Initialize(const FCaptureProtocolInitSettings& InSettings, const ICaptureProtocolHost& Host)
 {
+	InitSettings = InSettings;
 	if (!FFrameGrabberProtocol::Initialize(InSettings, Host))
 	{
 		return false;
 	}
 
-	UVideoCaptureSettings* CaptureSettings = CastChecked<UVideoCaptureSettings>(InSettings.ProtocolSettings);
+	ConditionallyCreateWriter(Host);
 
+	return AVIWriters.Num() && AVIWriters.Last()->IsCapturing();
+}
+
+void FVideoCaptureProtocol::ConditionallyCreateWriter(const ICaptureProtocolHost& Host)
+{
 #if PLATFORM_MAC
 	static const TCHAR* Extension = TEXT(".mov");
 #else
 	static const TCHAR* Extension = TEXT(".avi");
 #endif
 
+	FString VideoFilename = Host.GenerateFilename(FFrameMetrics(), Extension);
+
+	if (AVIWriters.Num() && VideoFilename == AVIWriters.Last()->Options.OutputFilename)
+	{
+		return;
+	}
+
+	Host.EnsureFileWritable(VideoFilename);
+
+	UVideoCaptureSettings* CaptureSettings = CastChecked<UVideoCaptureSettings>(InitSettings->ProtocolSettings);
+
 	FAVIWriterOptions Options;
-	Options.OutputFilename = Host.GenerateFilename(FFrameMetrics(), Extension);
+	Options.OutputFilename = MoveTemp(VideoFilename);
 	Options.CaptureFPS = Host.GetCaptureFrequency();
 	Options.CodecName = CaptureSettings->VideoCodec;
 	Options.bSynchronizeFrames = Host.GetCaptureStrategy().ShouldSynchronizeFrames();
-	Options.Width = InSettings.DesiredSize.X;
-	Options.Height = InSettings.DesiredSize.Y;
+	Options.Width = InitSettings->DesiredSize.X;
+	Options.Height = InitSettings->DesiredSize.Y;
 
 	if (CaptureSettings->bUseCompression)
 	{
@@ -40,21 +57,23 @@ bool FVideoCaptureProtocol::Initialize(const FCaptureProtocolInitSettings& InSet
 		Options.CompressionQuality = FMath::Clamp<float>(Options.CompressionQuality.GetValue(), 0.f, 1.f);
 	}
 
-	AVIWriter.Reset(FAVIWriter::CreateInstance(Options));
-	AVIWriter->Initialize();
-
-	return AVIWriter->IsCapturing();
+	AVIWriters.Emplace(FAVIWriter::CreateInstance(Options));
+	AVIWriters.Last()->Initialize();
 }
 
 struct FVideoFrameData : IFramePayload
 {
 	FFrameMetrics Metrics;
+	int32 WriterIndex;
 };
 
-FFramePayloadPtr FVideoCaptureProtocol::GetFramePayload(const FFrameMetrics& FrameMetrics, const ICaptureProtocolHost& Host) const
+FFramePayloadPtr FVideoCaptureProtocol::GetFramePayload(const FFrameMetrics& FrameMetrics, const ICaptureProtocolHost& Host)
 {
+	ConditionallyCreateWriter(Host);
+
 	TSharedRef<FVideoFrameData, ESPMode::ThreadSafe> FrameData = MakeShareable(new FVideoFrameData);
 	FrameData->Metrics = FrameMetrics;
+	FrameData->WriterIndex = AVIWriters.Num() - 1;
 	return FrameData;
 }
 
@@ -62,14 +81,54 @@ void FVideoCaptureProtocol::ProcessFrame(FCapturedFrameData Frame)
 {
 	FVideoFrameData* Payload = Frame.GetPayload<FVideoFrameData>();
 
-	AVIWriter->DropFrames(Payload->Metrics.NumDroppedFrames);
-	AVIWriter->Update(Payload->Metrics.TotalElapsedTime, MoveTemp(Frame.ColorBuffer));
+	const int32 WriterIndex = Payload->WriterIndex;
+
+	AVIWriters[WriterIndex]->DropFrames(Payload->Metrics.NumDroppedFrames);
+	AVIWriters[WriterIndex]->Update(Payload->Metrics.TotalElapsedTime, MoveTemp(Frame.ColorBuffer));
+
+	// Finalize previous writers if necessary
+	for (int32 Index = 0; Index < Payload->WriterIndex; ++Index)
+	{
+		TUniquePtr<FAVIWriter>& Writer = AVIWriters[Index];
+		if (Writer->IsCapturing())
+		{
+			Writer->Finalize();
+		}
+	}
 }
 
 void FVideoCaptureProtocol::Finalize()
 {
-	AVIWriter->Finalize();
-	AVIWriter.Reset();
+	for (TUniquePtr<FAVIWriter>& Writer : AVIWriters)
+	{
+		if (Writer->IsCapturing())
+		{
+			Writer->Finalize();
+		}
+	}
+	
+	AVIWriters.Empty();
 
 	FFrameGrabberProtocol::Finalize();
-}
\ No newline at end of file
+}
+
+bool FVideoCaptureProtocol::CanWriteToFile(const TCHAR* InFilename, bool bOverwriteExisting) const
+{
+	// When recording video, if the filename changes (ie due to the shot changing), we create new AVI writers.
+	// If we're not overwriting existing filenames we need to check if we're already recording a video of that name,
+	// before we can deem whether we can write to a new file (we can always write to a filename we're already writing to)
+	if (!bOverwriteExisting)
+	{
+		for (const TUniquePtr<FAVIWriter>& Writer : AVIWriters)
+		{
+			if (Writer->Options.OutputFilename == InFilename)
+			{
+				return true;
+			}
+		}
+
+		return IFileManager::Get().FileSize(InFilename) == -1;
+	}
+
+	return true;
+}
diff --git a/Engine/Source/Runtime/MovieSceneCapture/Public/AutomatedLevelSequenceCapture.h b/Engine/Source/Runtime/MovieSceneCapture/Public/AutomatedLevelSequenceCapture.h
--- a/Engine/Source/Runtime/MovieSceneCapture/Public/AutomatedLevelSequenceCapture.h
+++ b/Engine/Source/Runtime/MovieSceneCapture/Public/AutomatedLevelSequenceCapture.h
@@ -46,18 +46,45 @@ public:
 	UPROPERTY(config, EditAnywhere, Category=Animation, AdvancedDisplay, meta=(Units=Seconds, ClampMin=0))
 	float DelayBeforeWarmUp;
 
+	UPROPERTY(EditAnywhere, Category=CaptureSettings, AdvancedDisplay, meta=(EditInline))
+	ULevelSequenceBurnInOptions* BurnInOptions;
+
+	/** Whether to write edit decision lists (EDLs) if the sequence contains shots */
+	UPROPERTY(config, EditAnywhere, Category=Sequence)
+	bool bWriteEditDecisionList;
+
 public:
 	// UMovieSceneCapture interface
 	virtual void Initialize(TSharedPtr<FSceneViewport> InViewport, int32 PIEInstance = -1) override;
 
+	virtual void LoadFromConfig() override;
+
+	virtual void SaveToConfig() override;
+
+protected:
+
+	virtual void AddFormatMappings(TMap<FString, FStringFormatArg>& OutFormatMappings, const FFrameMetrics& FrameMetrics) const override;
+
+	/** Custom, additional json serialization */
+	virtual void SerializeAdditionalJson(FJsonObject& Object);
+
+	/** Custom, additional json deserialization */
+	virtual void DeserializeAdditionalJson(const FJsonObject& Object);
+
 private:
 
+	/** Update any cached information we need from the level sequence actor */
+	void UpdateFrameState();
+
 	/** Called when the level sequence has updated the world */
 	void SequenceUpdated(const ULevelSequencePlayer& Player, float CurrentTime, float PreviousTime);
 
 	/** Called to set up the player's playback range */
 	void SetupFrameRange();
 
+	/** Export EDL if requested */
+	void ExportEDL();
+
 	/** Delegate binding for the above callback */
 	FDelegateHandle OnPlayerUpdatedBinding;
 
@@ -88,10 +115,8 @@ private:
 
 	/** The number of warm up frames left before we actually start saving out images */
 	int32 RemainingWarmUpFrames;
-
-	/** The delta of the last sequence updat */
-	float LastSequenceUpdateDelta;
 	
+	FLevelSequencePlayerSnapshot CachedState;
 #endif
 };
 
diff --git a/Engine/Source/Runtime/MovieSceneCapture/Public/IMovieSceneCaptureProtocol.h b/Engine/Source/Runtime/MovieSceneCapture/Public/IMovieSceneCaptureProtocol.h
--- a/Engine/Source/Runtime/MovieSceneCapture/Public/IMovieSceneCaptureProtocol.h
+++ b/Engine/Source/Runtime/MovieSceneCapture/Public/IMovieSceneCaptureProtocol.h
@@ -52,6 +52,16 @@ struct IMovieSceneCaptureProtocol
 	 */
 	virtual void AddFormatMappings(TMap<FString, FStringFormatArg>& FormatMappings) const {}
 
+	/**
+	 * Test whether this capture protocol thinks the file should be written to. Only called when we're not overwriting existing files.
+	 * By default, we simply test for the file's existence, however this can be overridden to afford complex behaviour like
+	 * writing out multiple video files for different file names
+	 * @param InFilename 			The filename to test
+	 * @param bOverwriteExisting	Whether we are allowed to overwrite existing files
+	 * @return Whether we should deem this file writable or not
+	 */
+	virtual MOVIESCENECAPTURE_API bool CanWriteToFile(const TCHAR* InFilename, bool bOverwriteExisting) const;
+
 	virtual ~IMovieSceneCaptureProtocol() {}
 };
 
@@ -101,6 +111,9 @@ struct ICaptureProtocolHost
 	/** Generate a filename for the specified frame metrics. How often this is called is determined by the protocol itself */
 	virtual FString GenerateFilename(const FFrameMetrics& FrameMetrics, const TCHAR* Extension) const = 0;
 
+	/** Ensure that the specified file is writable, potentially deleting an existing file if settings allow */
+	virtual void EnsureFileWritable(const FString& File) const = 0;
+
 	/** Get the capture frequency */
 	virtual float GetCaptureFrequency() const = 0;
 
diff --git a/Engine/Source/Runtime/MovieSceneCapture/Public/MovieSceneCapture.h b/Engine/Source/Runtime/MovieSceneCapture/Public/MovieSceneCapture.h
--- a/Engine/Source/Runtime/MovieSceneCapture/Public/MovieSceneCapture.h
+++ b/Engine/Source/Runtime/MovieSceneCapture/Public/MovieSceneCapture.h
@@ -12,6 +12,8 @@
 #include "MovieSceneCaptureProtocolSettings.h"
 #include "MovieSceneCapture.generated.h"
 
+class FJsonObject;
+
 /** Structure used to cache various metrics for our capture */
 struct FCachedMetrics
 {
@@ -46,6 +48,24 @@ public:
 	const FMovieSceneCaptureSettings& GetSettings() const override { return Settings; }
 	// End IMovieSceneCaptureInterface
 
+	/** Load save from config helpers */
+	virtual void LoadFromConfig();
+	virtual void SaveToConfig();
+
+	/** Serialize additional json data for this capture */
+	void SerializeJson(FJsonObject& Object);
+
+	/** Deserialize additional json data for this capture */
+	void DeserializeJson(const FJsonObject& Object);
+
+protected:
+
+	/** Custom, additional json serialization */
+	virtual void SerializeAdditionalJson(FJsonObject& Object){}
+
+	/** Custom, additional json deserialization */
+	virtual void DeserializeAdditionalJson(const FJsonObject& Object){}
+
 public:
 
 	/** The type of capture protocol to use */
@@ -118,12 +138,19 @@ protected:
 
 	/**~ ICaptureProtocolHost interface */
 	virtual FString GenerateFilename(const FFrameMetrics& FrameMetrics, const TCHAR* Extension) const override;
+	virtual void EnsureFileWritable(const FString& File) const override;
 	virtual float GetCaptureFrequency() const { return Settings.FrameRate; }
 	virtual const ICaptureStrategy& GetCaptureStrategy() const { return *CaptureStrategy; }
 
+	/** Add additional format mappings to be used when generating filenames */
+	virtual void AddFormatMappings(TMap<FString, FStringFormatArg>& OutFormatMappings, const FFrameMetrics& FrameMetrics) const {}
+
 	/** Resolve the specified format using the user supplied formatting rules. */
 	FString ResolveFileFormat(const FString& Format, const FFrameMetrics& FrameMetrics) const;
 
+	/** Initialize the settings structure for the current capture type */
+	void InitializeSettings();
+
 protected:
 
 #if WITH_EDITOR
@@ -144,7 +171,7 @@ protected:
 	/** Cached metrics for this capture operation */
 	FCachedMetrics CachedMetrics;
 	/** Format mappings used for generating filenames */
-	mutable TMap<FString, FStringFormatArg> FormatMappings;
+	TMap<FString, FStringFormatArg> FormatMappings;
 	/** The number of frames to capture.  If this is zero, we'll capture the entire sequence. */
 	int32 FrameCount;
 	/** Whether we have started capturing or not */
@@ -154,6 +181,8 @@ protected:
 	int32 FrameNumberOffset;
 	/** Event that is triggered when capturing has finished */
 	FOnCaptureFinished OnCaptureFinishedDelegate;
+	/** Format string used for frame numbers */
+	FString FrameNumberFormat;
 };
 
 /** A strategy that employs a fixed frame time-step, and as such never drops a frame. Potentially accelerated. */
diff --git a/Engine/Source/Runtime/MovieSceneCapture/Public/MovieSceneCaptureHelpers.h b/Engine/Source/Runtime/MovieSceneCapture/Public/MovieSceneCaptureHelpers.h
new file mode 100644
--- /dev/null
+++ b/Engine/Source/Runtime/MovieSceneCapture/Public/MovieSceneCaptureHelpers.h
@@ -0,0 +1,30 @@
+// Copyright 1998-2016 Epic Games, Inc. All Rights Reserved.
+
+#pragma once
+
+class UMovieScene;
+
+class MOVIESCENECAPTURE_API MovieSceneCaptureHelpers
+{
+public:
+
+	/**
+	 * Import EDL
+	 *
+	 * @param InMovieScene The movie scene to import the edl into
+	 * @param InFrameRate The frame rate to import the EDL at
+	 * @param InFilename The filename to import
+	 * @return Whether the import was successful
+	 */
+	static bool ImportEDL(UMovieScene* InMovieScene, float InFrameRate, FString InFilename);
+
+	/**
+	 * Export EDL
+	 *
+	 * @param InMovieScene The movie scene with the cinematic shot track and audio tracks to export
+	 * @param InFrameRate The frame rate to export the EDL at
+	 * @param InSaveDirectory The directory path to save to.
+	 * @return Whether the export was successful
+	 */
+	static bool ExportEDL(const UMovieScene* InMovieScene, float InFrameRate, FString InSaveDirectory);
+};
\ No newline at end of file
diff --git a/Engine/Source/Runtime/MovieSceneCapture/Public/MovieSceneCaptureSettings.h b/Engine/Source/Runtime/MovieSceneCapture/Public/MovieSceneCaptureSettings.h
--- a/Engine/Source/Runtime/MovieSceneCapture/Public/MovieSceneCaptureSettings.h
+++ b/Engine/Source/Runtime/MovieSceneCapture/Public/MovieSceneCaptureSettings.h
@@ -61,6 +61,10 @@ struct MOVIESCENECAPTURE_API FMovieSceneCaptureSettings
 	UPROPERTY(config, EditAnywhere, Category=General, AdvancedDisplay)
 	bool bUseRelativeFrameNumbers;
 
+	/** How much to zero-pad frame numbers on filenames */
+	UPROPERTY(config)
+	uint8 ZeroPadFrameNumbers;
+
 	/** The frame rate at which to capture */
 	UPROPERTY(config, EditAnywhere, Category=CaptureSettings)
 	int32 FrameRate;
diff --git a/Engine/Source/Runtime/MovieSceneCapture/Public/Protocols/CompositionGraphCaptureProtocol.h b/Engine/Source/Runtime/MovieSceneCapture/Public/Protocols/CompositionGraphCaptureProtocol.h
--- a/Engine/Source/Runtime/MovieSceneCapture/Public/Protocols/CompositionGraphCaptureProtocol.h
+++ b/Engine/Source/Runtime/MovieSceneCapture/Public/Protocols/CompositionGraphCaptureProtocol.h
@@ -6,6 +6,18 @@
 #include "MovieSceneCaptureProtocolSettings.h"
 #include "CompositionGraphCaptureProtocol.generated.h"
 
+/** Used by UCompositionGraphCaptureSettings. Matches gamut oreder in TonemapCommon.usf OuputGamutMappingMatrix()*/
+UENUM()
+enum EHDRCaptureGamut
+{
+	HCGM_Rec709 UMETA(DisplayName = "Rec.709 / sRGB"),
+	HCGM_P3DCI UMETA(DisplayName = "P3 D65"),
+	HCGM_Rec2020 UMETA(DisplayName = "Rec.2020"),
+	HCGM_ACES UMETA(DisplayName = "ACES"),
+	HCGM_ACEScg UMETA(DisplayName = "ACEScg"),
+	HCGM_MAX,
+};
+
 USTRUCT()
 struct MOVIESCENECAPTURE_API FCompositionGraphCapturePasses
 {
@@ -33,6 +45,10 @@ public:
 	UPROPERTY(config, EditAnywhere, Category="Composition Graph Options")
 	bool bCaptureFramesInHDR;
 
+	/** The color gamut to use when storing HDR captured data. The gamut depends on whether the bCaptureFramesInHDR option is enabled. */
+	UPROPERTY(config, EditAnywhere, Category="Composition Graph Options", meta = (EditCondition = "bCaptureFramesInHDR"))
+	TEnumAsByte<enum EHDRCaptureGamut> CaptureGamut;
+
 	/** Custom post processing material to use for rendering */
 	UPROPERTY(config, EditAnywhere, Category="Composition Graph Options", meta=(AllowedClasses=""))
 	FStringAssetReference PostProcessingMaterial;
diff --git a/Engine/Source/Runtime/MovieSceneCapture/Public/Protocols/FrameGrabberProtocol.h b/Engine/Source/Runtime/MovieSceneCapture/Public/Protocols/FrameGrabberProtocol.h
--- a/Engine/Source/Runtime/MovieSceneCapture/Public/Protocols/FrameGrabberProtocol.h
+++ b/Engine/Source/Runtime/MovieSceneCapture/Public/Protocols/FrameGrabberProtocol.h
@@ -46,7 +46,7 @@ protected:
 
 	 * @return Shared pointer to a payload to associate with the frame, or nullptr
 	 */
-	virtual FFramePayloadPtr GetFramePayload(const FFrameMetrics& FrameMetrics, const ICaptureProtocolHost& Host) const = 0;
+	virtual FFramePayloadPtr GetFramePayload(const FFrameMetrics& FrameMetrics, const ICaptureProtocolHost& Host) = 0;
 
 	/**
 	 * Process a captured frame. This may be called on any thread.
diff --git a/Engine/Source/Runtime/MovieSceneCapture/Public/Protocols/ImageSequenceProtocol.h b/Engine/Source/Runtime/MovieSceneCapture/Public/Protocols/ImageSequenceProtocol.h
--- a/Engine/Source/Runtime/MovieSceneCapture/Public/Protocols/ImageSequenceProtocol.h
+++ b/Engine/Source/Runtime/MovieSceneCapture/Public/Protocols/ImageSequenceProtocol.h
@@ -85,7 +85,7 @@ struct MOVIESCENECAPTURE_API FImageSequenceProtocol : FFrameGrabberProtocol
 
 	/** ~FFrameGrabberProtocol implementation */
 	virtual bool Initialize(const FCaptureProtocolInitSettings& InSettings, const ICaptureProtocolHost& Host) override;
-	virtual FFramePayloadPtr GetFramePayload(const FFrameMetrics& FrameMetrics, const ICaptureProtocolHost& Host) const;
+	virtual FFramePayloadPtr GetFramePayload(const FFrameMetrics& FrameMetrics, const ICaptureProtocolHost& Host);
 	virtual void ProcessFrame(FCapturedFrameData Frame);
 	virtual void AddFormatMappings(TMap<FString, FStringFormatArg>& FormatMappings) const override;
 	virtual void Finalize() override;
diff --git a/Engine/Source/Runtime/MovieSceneCapture/Public/Protocols/VideoCaptureProtocol.h b/Engine/Source/Runtime/MovieSceneCapture/Public/Protocols/VideoCaptureProtocol.h
--- a/Engine/Source/Runtime/MovieSceneCapture/Public/Protocols/VideoCaptureProtocol.h
+++ b/Engine/Source/Runtime/MovieSceneCapture/Public/Protocols/VideoCaptureProtocol.h
@@ -28,9 +28,16 @@ struct MOVIESCENECAPTURE_API FVideoCaptureProtocol : FFrameGrabberProtocol
 {
 	virtual bool Initialize(const FCaptureProtocolInitSettings& InSettings, const ICaptureProtocolHost& Host) override;
 	virtual void Finalize() override;
-	virtual FFramePayloadPtr GetFramePayload(const FFrameMetrics& FrameMetrics, const ICaptureProtocolHost& Host) const;
+	virtual FFramePayloadPtr GetFramePayload(const FFrameMetrics& FrameMetrics, const ICaptureProtocolHost& Host);
 	virtual void ProcessFrame(FCapturedFrameData Frame);
-	
+	virtual bool CanWriteToFile(const TCHAR* InFilename, bool bOverwriteExisting) const override;
+
+protected:
+
+	void ConditionallyCreateWriter(const ICaptureProtocolHost& Host);
 private:
-	TUniquePtr<FAVIWriter> AVIWriter;
+
+	TOptional<FCaptureProtocolInitSettings> InitSettings;
+
+	TArray<TUniquePtr<FAVIWriter>> AVIWriters;
 };
\ No newline at end of file
diff --git a/Engine/Source/Runtime/MovieSceneTracks/Private/Sections/MovieSceneBoolSection.cpp b/Engine/Source/Runtime/MovieSceneTracks/Private/Sections/MovieSceneBoolSection.cpp
--- a/Engine/Source/Runtime/MovieSceneTracks/Private/Sections/MovieSceneBoolSection.cpp
+++ b/Engine/Source/Runtime/MovieSceneTracks/Private/Sections/MovieSceneBoolSection.cpp
@@ -6,15 +6,23 @@
 
 UMovieSceneBoolSection::UMovieSceneBoolSection( const FObjectInitializer& ObjectInitializer )
 	: Super( ObjectInitializer )
-	, DefaultValue(false)
+	, DefaultValue_DEPRECATED(false)
 {
 	SetIsInfinite(true);
 }
 
+void UMovieSceneBoolSection::PostLoad()
+{
+	if (GetCurve().GetDefaultValue() == MAX_int32 && DefaultValue_DEPRECATED)
+	{
+		GetCurve().SetDefaultValue(DefaultValue_DEPRECATED);
+	}
+	Super::PostLoad();
+}
 
 bool UMovieSceneBoolSection::Eval( float Position ) const
 {
-	return !!BoolCurve.Evaluate(Position, DefaultValue ? 1 : 0);
+	return !!BoolCurve.Evaluate(Position);
 }
 
 
diff --git a/Engine/Source/Runtime/MovieSceneTracks/Private/Sections/MovieSceneColorSection.cpp b/Engine/Source/Runtime/MovieSceneTracks/Private/Sections/MovieSceneColorSection.cpp
--- a/Engine/Source/Runtime/MovieSceneTracks/Private/Sections/MovieSceneColorSection.cpp
+++ b/Engine/Source/Runtime/MovieSceneTracks/Private/Sections/MovieSceneColorSection.cpp
@@ -12,7 +12,14 @@ void FMovieSceneColorKeyStruct::PropagateChanges(const FPropertyChangedEvent& Ch
 {
 	for (int32 Index = 0; Index <= 3; ++Index)
 	{
-		Keys[Index]->Value = Color.Component(Index);
+		if (Keys[Index] == nullptr && Curves[Index] != nullptr)
+		{
+			Curves[Index]->SetDefaultValue(Color.Component(Index));
+		}
+		else
+		{
+			Keys[Index]->Value = Color.Component(Index);
+		}
 	}
 }
 
@@ -114,15 +121,35 @@ TSharedPtr<FStructOnScope> UMovieSceneColorSection::GetKeyStruct(const TArray<FK
 	TSharedRef<FStructOnScope> KeyStruct = MakeShareable(new FStructOnScope(FMovieSceneColorKeyStruct::StaticStruct()));
 	auto Struct = (FMovieSceneColorKeyStruct*)KeyStruct->GetStructMemory();
 	{
+		Struct->Curves[0] = &RedCurve;
+		Struct->Curves[1] = &GreenCurve;
+		Struct->Curves[2] = &BlueCurve;
+		Struct->Curves[3] = &AlphaCurve;
+
 		Struct->Keys[0] = RedCurve.GetFirstMatchingKey(KeyHandles);
 		Struct->Keys[1] = GreenCurve.GetFirstMatchingKey(KeyHandles);
 		Struct->Keys[2] = BlueCurve.GetFirstMatchingKey(KeyHandles);
 		Struct->Keys[3] = AlphaCurve.GetFirstMatchingKey(KeyHandles);
 
+		float FirstValidKeyTime = 0.f;
+		for (int32 Index = 0; Index <= 3; ++Index)
+		{
+			if (Struct->Keys[Index] != nullptr)
+			{
+				FirstValidKeyTime = Struct->Keys[Index]->Time;
+			}
+		}
+
 		for (int32 Index = 0; Index <= 3; ++Index)
 		{
-			check(Struct->Keys[Index] != nullptr);
-			Struct->Color.Component(Index) = Struct->Keys[Index]->Value;
+			if (Struct->Keys[Index] == nullptr && Struct->Curves[Index] != nullptr)
+			{
+				Struct->Color.Component(Index) = Struct->Curves[Index]->Eval(FirstValidKeyTime);
+			}
+			else
+			{
+				Struct->Color.Component(Index) = Struct->Keys[Index]->Value;
+			}
 		}
 	}
 
diff --git a/Engine/Source/Runtime/MovieSceneTracks/Private/TrackInstances/MovieSceneSkeletalAnimationTrackInstance.cpp b/Engine/Source/Runtime/MovieSceneTracks/Private/TrackInstances/MovieSceneSkeletalAnimationTrackInstance.cpp
--- a/Engine/Source/Runtime/MovieSceneTracks/Private/TrackInstances/MovieSceneSkeletalAnimationTrackInstance.cpp
+++ b/Engine/Source/Runtime/MovieSceneTracks/Private/TrackInstances/MovieSceneSkeletalAnimationTrackInstance.cpp
@@ -14,8 +14,6 @@ float MapTimeToAnimation(float ThisPosition, UMovieSceneSkeletalAnimationSection
 {
 	ThisPosition = FMath::Clamp(ThisPosition, AnimSection->GetStartTime(), AnimSection->GetEndTime());
 
-	ThisPosition -= 1 / 1000.0f;
-
 	float AnimPlayRate = FMath::IsNearlyZero(AnimSection->GetPlayRate()) ? 1.0f : AnimSection->GetPlayRate();
 	float SeqLength = AnimSection->GetSequenceLength() - (AnimSection->GetStartOffset() + AnimSection->GetEndOffset());
 
@@ -114,79 +112,107 @@ void FMovieSceneSkeletalAnimationTrackInstance::RestoreState(const TArray<TWeakO
 	}
 }
 
-void FMovieSceneSkeletalAnimationTrackInstance::Update(EMovieSceneUpdateData& UpdateData, const TArray<TWeakObjectPtr<UObject>>& RuntimeObjects, class IMovieScenePlayer& Player, FMovieSceneSequenceInstance& SequenceInstance) 
+USkeletalMeshComponent* GetSkeletalMeshComponentFromRuntimeObjectPtr( TWeakObjectPtr<UObject> RuntimeObjectPtr )
 {
-	// @todo Sequencer gameplay update has a different code path than editor update for animation
-
-	for (int32 i = 0; i < RuntimeObjects.Num(); ++i)
+	USkeletalMeshComponent* SkeletalMeshComponent = nullptr;
+	UObject* RuntimeObject = RuntimeObjectPtr.Get();
+	if ( RuntimeObject != nullptr )
 	{
-		// get the skeletal mesh component we want to control
-		USkeletalMeshComponent* Component = nullptr;
-		UObject* RuntimeObject = RuntimeObjects[i].Get();
+		// first try to control the component directly
+		SkeletalMeshComponent = Cast<USkeletalMeshComponent>( RuntimeObject );
+		if ( SkeletalMeshComponent == nullptr )
+		{
+			// then check to see if we are controlling an actor & if so use its first USkeletalMeshComponent 
+			AActor* Actor = Cast<AActor>( RuntimeObject );
+			if ( Actor != nullptr )
+			{
+				SkeletalMeshComponent = Actor->FindComponentByClass<USkeletalMeshComponent>();
+			}
+		}
+	}
+	return SkeletalMeshComponent;
+}
 
-		if (RuntimeObject != nullptr)
+void FMovieSceneSkeletalAnimationTrackInstance::Update( EMovieSceneUpdateData& UpdateData, const TArray<TWeakObjectPtr<UObject>>& RuntimeObjects, class IMovieScenePlayer& Player, FMovieSceneSequenceInstance& SequenceInstance )
+{
+	// @todo Sequencer gameplay update has a different code path than editor update for animation
+	for ( TWeakObjectPtr<UObject> RuntimeObjectPtr : RuntimeObjects )
+	{
+		USkeletalMeshComponent* SkeletalMeshComponent = GetSkeletalMeshComponentFromRuntimeObjectPtr( RuntimeObjectPtr );
+		if ( SkeletalMeshComponent )
 		{
-			// first try to control the component directly
-			USkeletalMeshComponent* SkeletalMeshComponent = Cast<USkeletalMeshComponent>(RuntimeObject);
-			if (SkeletalMeshComponent == nullptr)
+			UMovieSceneSkeletalAnimationSection* AnimSection = Cast<UMovieSceneSkeletalAnimationSection>( AnimationTrack->GetAnimSectionAtTime( UpdateData.Position ) );
+
+			// cbb: If there is no overlapping section, evaluate the closest section only if the current time is before it.
+			if ( AnimSection == nullptr )
 			{
-				// then check to see if we are controlling an actor & if so use its first USkeletalMeshComponent 
-				AActor* Actor = Cast<AActor>(RuntimeObject);
-				if(Actor != nullptr)
-				{
-					SkeletalMeshComponent = Actor->FindComponentByClass<USkeletalMeshComponent>();
-				}
+				AnimSection = Cast<UMovieSceneSkeletalAnimationSection>( MovieSceneHelpers::FindNearestSectionAtTime( AnimationTrack->GetAllSections(), UpdateData.Position ) );
 			}
 
-			if (SkeletalMeshComponent)
+			if ( AnimSection && AnimSection->IsActive() )
 			{
-				UMovieSceneSkeletalAnimationSection* AnimSection = Cast<UMovieSceneSkeletalAnimationSection>(AnimationTrack->GetAnimSectionAtTime(UpdateData.Position));
-			
-				// cbb: If there is no overlapping section, evaluate the closest section only if the current time is before it.
-				if (AnimSection == nullptr)
-				{
-					AnimSection = Cast<UMovieSceneSkeletalAnimationSection>(MovieSceneHelpers::FindNearestSectionAtTime(AnimationTrack->GetAllSections(), UpdateData.Position));
-				}
+				UAnimSequence* AnimSequence = AnimSection->GetAnimSequence();
 
-				if (AnimSection && AnimSection->IsActive())
+				if ( AnimSequence )
 				{
-					UAnimSequence* AnimSequence = AnimSection->GetAnimSequence();
+					float EvalTime = MapTimeToAnimation( UpdateData.Position, AnimSection );
 
-					if (AnimSequence)
-					{
-						float EvalTime = MapTimeToAnimation(UpdateData.Position, AnimSection);
+					int32 ChannelIndex = 0;
 
-						int32 ChannelIndex = 0;
+					const bool bLooping = false;
 
-						const bool bLooping = false;
+					if ( ShouldUsePreviewPlayback( Player, SkeletalMeshComponent ) )
+					{
+						// If the playback status is jumping, ie. one such occurrence is setting the time for thumbnail generation, disable anim notifies updates because it could fire audio
+						const bool bFireNotifies = Player.GetPlaybackStatus() == EMovieScenePlayerStatus::Jumping || Player.GetPlaybackStatus() == EMovieScenePlayerStatus::Stopped ? false : true;
 
-						if (ShouldUsePreviewPlayback(Player, RuntimeObject))
-						{
-							// If the playback status is jumping, ie. one such occurrence is setting the time for thumbnail generation, disable anim notifies updates because it could fire audio
-							const bool bFireNotifies = Player.GetPlaybackStatus() == EMovieScenePlayerStatus::Jumping || Player.GetPlaybackStatus() == EMovieScenePlayerStatus::Stopped ? false : true;
-
-							float DeltaTime = UpdateData.Position > UpdateData.LastPosition ? UpdateData.Position - UpdateData.LastPosition : 0.f;
-
-							// When jumping from one cut to another cut, the delta time should be 0 so that anim notifies before the current position are not evaluated. Note, anim notifies at the current time should still be evaluated.
-							if (UpdateData.bJumpCut)
-							{
-								DeltaTime = 0.f;
-							}
-
-							const bool bResetDynamics = Player.GetPlaybackStatus() == EMovieScenePlayerStatus::Stepping || 
-														Player.GetPlaybackStatus() == EMovieScenePlayerStatus::Jumping || 
-														Player.GetPlaybackStatus() == EMovieScenePlayerStatus::Scrubbing || 
-														(DeltaTime == 0.0f && Player.GetPlaybackStatus() != EMovieScenePlayerStatus::Stopped); 
-						
-							PreviewSetAnimPosition(SkeletalMeshComponent, AnimSection->GetSlotName(), ChannelIndex, AnimSequence, EvalTime, bLooping, bFireNotifies, DeltaTime, Player.GetPlaybackStatus() == EMovieScenePlayerStatus::Playing, bResetDynamics);
-						}
-						else
+						float DeltaTime = UpdateData.Position > UpdateData.LastPosition ? UpdateData.Position - UpdateData.LastPosition : 0.f;
+
+						// When jumping from one cut to another cut, the delta time should be 0 so that anim notifies before the current position are not evaluated. Note, anim notifies at the current time should still be evaluated.
+						if ( UpdateData.bJumpCut )
 						{
-							// Don't fire notifies at runtime since they will be fired through the standard animation tick path.
-							const bool bFireNotifies = false;
-							SetAnimPosition(SkeletalMeshComponent, AnimSection->GetSlotName(), ChannelIndex, AnimSequence, EvalTime, bLooping, bFireNotifies);
+							DeltaTime = 0.f;
 						}
+
+						const bool bResetDynamics = Player.GetPlaybackStatus() == EMovieScenePlayerStatus::Stepping ||
+							Player.GetPlaybackStatus() == EMovieScenePlayerStatus::Jumping ||
+							Player.GetPlaybackStatus() == EMovieScenePlayerStatus::Scrubbing ||
+							( DeltaTime == 0.0f && Player.GetPlaybackStatus() != EMovieScenePlayerStatus::Stopped );
+
+						PreviewSetAnimPosition( SkeletalMeshComponent, AnimSection->GetSlotName(), ChannelIndex, AnimSequence, EvalTime, bLooping, bFireNotifies, DeltaTime, Player.GetPlaybackStatus() == EMovieScenePlayerStatus::Playing, bResetDynamics );
+					}
+					else
+					{
+						// Don't fire notifies at runtime since they will be fired through the standard animation tick path.
+						const bool bFireNotifies = false;
+						SetAnimPosition( SkeletalMeshComponent, AnimSection->GetSlotName(), ChannelIndex, AnimSequence, EvalTime, bLooping, bFireNotifies );
+					}
+				}
+			}
+		}
+	}
+}
+
+void FMovieSceneSkeletalAnimationTrackInstance::RefreshInstance( const TArray<TWeakObjectPtr<UObject>>& RuntimeObjects, IMovieScenePlayer& Player, FMovieSceneSequenceInstance& SequenceInstance )
+{
+	// When not in preview playback we need to stop any running animations to prevent the animations from being advanced a frame when
+	// they are ticked by the engine.
+	for ( TWeakObjectPtr<UObject> RuntimeObjectPtr : RuntimeObjects )
+	{
+		if ( ShouldUsePreviewPlayback( Player, RuntimeObjectPtr.Get() ) == false )
+		{
+			USkeletalMeshComponent* SkeletalMeshComponent = GetSkeletalMeshComponentFromRuntimeObjectPtr( RuntimeObjectPtr );
+			if ( SkeletalMeshComponent != nullptr )
+			{
+				UAnimInstance* AnimInstance = SkeletalMeshComponent->GetAnimInstance();
+				if ( AnimInstance )
+				{
+					UAnimSingleNodeInstance * SingleNodeInstance = SkeletalMeshComponent->GetSingleNodeInstance();
+					if ( SingleNodeInstance )
+					{
+						SingleNodeInstance->SetPlaying( false );
 					}
+					// TODO: Anim montage?
 				}
 			}
 		}
diff --git a/Engine/Source/Runtime/MovieSceneTracks/Private/TrackInstances/MovieSceneSlomoTrackInstance.cpp b/Engine/Source/Runtime/MovieSceneTracks/Private/TrackInstances/MovieSceneSlomoTrackInstance.cpp
--- a/Engine/Source/Runtime/MovieSceneTracks/Private/TrackInstances/MovieSceneSlomoTrackInstance.cpp
+++ b/Engine/Source/Runtime/MovieSceneTracks/Private/TrackInstances/MovieSceneSlomoTrackInstance.cpp
@@ -57,11 +57,13 @@ void FMovieSceneSlomoTrackInstance::Update(EMovieSceneUpdateData& UpdateData, co
 
 	float FloatValue = 0.0f;
 
-	if (SlomoTrack->Eval(UpdateData.Position, UpdateData.LastPosition, FloatValue))
+	if (!SlomoTrack->Eval(UpdateData.Position, UpdateData.LastPosition, FloatValue) || (FloatValue <= 0.0f))
 	{
-		WorldSettings->MatineeTimeDilation = FloatValue;
-		WorldSettings->ForceNetUpdate();
+		return;
 	}
+
+	WorldSettings->MatineeTimeDilation = FloatValue;
+	WorldSettings->ForceNetUpdate();
 }
 
 
diff --git a/Engine/Source/Runtime/MovieSceneTracks/Private/TrackInstances/MovieSceneSubTrackInstance.cpp b/Engine/Source/Runtime/MovieSceneTracks/Private/TrackInstances/MovieSceneSubTrackInstance.cpp
--- a/Engine/Source/Runtime/MovieSceneTracks/Private/TrackInstances/MovieSceneSubTrackInstance.cpp
+++ b/Engine/Source/Runtime/MovieSceneTracks/Private/TrackInstances/MovieSceneSubTrackInstance.cpp
@@ -241,8 +241,8 @@ void FMovieSceneSubTrackInstance::UpdateSection( EMovieSceneUpdateData& UpdateDa
 		if (SubMovieScene && SubMovieScene->GetForceFixedFrameIntervalPlayback() && SubMovieScene->GetFixedFrameInterval() > 0 )
 		{
 			float FixedFrameInterval = ( SubMovieScene->GetFixedFrameInterval() / SubSection->TimeScale );
-			InstancePosition = FMath::RoundToInt( InstancePosition / FixedFrameInterval ) * FixedFrameInterval;
-			InstanceLastPosition = FMath::RoundToInt( InstanceLastPosition / FixedFrameInterval ) * FixedFrameInterval;
+			InstancePosition = UMovieScene::CalculateFixedFrameTime( InstancePosition, FixedFrameInterval );
+			InstanceLastPosition = UMovieScene::CalculateFixedFrameTime( InstanceLastPosition, FixedFrameInterval );
 		}
 	}
 
diff --git a/Engine/Source/Runtime/MovieSceneTracks/Private/Tracks/MovieSceneSpawnTrack.cpp b/Engine/Source/Runtime/MovieSceneTracks/Private/Tracks/MovieSceneSpawnTrack.cpp
--- a/Engine/Source/Runtime/MovieSceneTracks/Private/Tracks/MovieSceneSpawnTrack.cpp
+++ b/Engine/Source/Runtime/MovieSceneTracks/Private/Tracks/MovieSceneSpawnTrack.cpp
@@ -48,7 +48,7 @@ bool UMovieSceneSpawnTrack::Eval(float Position, float LastPostion, bool& bOutSp
 UMovieSceneSection* UMovieSceneSpawnTrack::CreateNewSection()
 {
 	UMovieSceneBoolSection* Section = NewObject<UMovieSceneBoolSection>(this, NAME_None, RF_Transactional);
-	Section->DefaultValue = true;
+	Section->GetCurve().SetDefaultValue(true);
 	return Section;
 }
 
@@ -94,6 +94,35 @@ const TArray<UMovieSceneSection*>& UMovieSceneSpawnTrack::GetAllSections() const
 	return Sections;
 }
 
+#if WITH_EDITORONLY_DATA
+
+ECookOptimizationFlags UMovieSceneSpawnTrack::GetCookOptimizationFlags() const
+{
+	// Since the spawn track denotes the lifetime of a spawnable, if the object is never spawned, we can remove the entire object
+	for (UMovieSceneSection* Section : Sections)
+	{
+		UMovieSceneBoolSection* BoolSection = CastChecked<UMovieSceneBoolSection>(Section);
+		if (!BoolSection->IsActive())
+		{
+			continue;
+		}
+		else if (BoolSection->GetCurve().GetNumKeys() == 0 && BoolSection->GetCurve().GetDefaultValue())
+		{
+			return ECookOptimizationFlags::None;
+		}
+		else for (auto It = BoolSection->GetCurve().GetKeyIterator(); It; ++It)
+		{
+			if (It->Value)
+			{
+				return ECookOptimizationFlags::None;
+			}
+		}
+	}
+
+	return ECookOptimizationFlags::RemoveObject;
+}
+
+#endif
 
 #if WITH_EDITORONLY_DATA
 
diff --git a/Engine/Source/Runtime/MovieSceneTracks/Public/Sections/MovieSceneBoolSection.h b/Engine/Source/Runtime/MovieSceneTracks/Public/Sections/MovieSceneBoolSection.h
--- a/Engine/Source/Runtime/MovieSceneTracks/Public/Sections/MovieSceneBoolSection.h
+++ b/Engine/Source/Runtime/MovieSceneTracks/Public/Sections/MovieSceneBoolSection.h
@@ -18,9 +18,9 @@ class UMovieSceneBoolSection
 
 public:
 
-	/** The default value to use when no keys are present */
-	UPROPERTY(EditAnywhere, Category="Curve")
-	bool DefaultValue;
+	/** The default value to use when no keys are present - use GetCurve().SetDefaultValue() */
+	UPROPERTY()
+	bool DefaultValue_DEPRECATED;
 
 public:
 
@@ -49,6 +49,8 @@ public:
 	/** Gets all the keys of this boolean section */
 	FIntegralCurve& GetCurve() { return BoolCurve; }
 
+	virtual void PostLoad() override;
+
 private:
 	/** Ordered curve data */
 	// @todo Sequencer This could be optimized by packing the bools separately
diff --git a/Engine/Source/Runtime/MovieSceneTracks/Public/Sections/MovieSceneColorSection.h b/Engine/Source/Runtime/MovieSceneTracks/Public/Sections/MovieSceneColorSection.h
--- a/Engine/Source/Runtime/MovieSceneTracks/Public/Sections/MovieSceneColorSection.h
+++ b/Engine/Source/Runtime/MovieSceneTracks/Public/Sections/MovieSceneColorSection.h
@@ -46,6 +46,7 @@ struct FMovieSceneColorKeyStruct
 	FLinearColor Color;
 
 	FRichCurveKey* Keys[4];
+	FRichCurve* Curves[4];
 
 	virtual void PropagateChanges(const FPropertyChangedEvent& ChangeEvent) override;
 };
diff --git a/Engine/Source/Runtime/MovieSceneTracks/Public/Sections/MovieSceneFadeSection.h b/Engine/Source/Runtime/MovieSceneTracks/Public/Sections/MovieSceneFadeSection.h
--- a/Engine/Source/Runtime/MovieSceneTracks/Public/Sections/MovieSceneFadeSection.h
+++ b/Engine/Source/Runtime/MovieSceneTracks/Public/Sections/MovieSceneFadeSection.h
@@ -21,7 +21,7 @@ class UMovieSceneFadeSection
 public:
 
 	/** Fade color */
-	UPROPERTY(EditAnywhere, Category="Fade")
+	UPROPERTY(EditAnywhere, Category="Fade", meta=(InlineColorPicker))
 	FLinearColor FadeColor;
 
 	/** Fade audio */
diff --git a/Engine/Source/Runtime/MovieSceneTracks/Public/TrackInstances/MovieSceneSkeletalAnimationTrackInstance.h b/Engine/Source/Runtime/MovieSceneTracks/Public/TrackInstances/MovieSceneSkeletalAnimationTrackInstance.h
--- a/Engine/Source/Runtime/MovieSceneTracks/Public/TrackInstances/MovieSceneSkeletalAnimationTrackInstance.h
+++ b/Engine/Source/Runtime/MovieSceneTracks/Public/TrackInstances/MovieSceneSkeletalAnimationTrackInstance.h
@@ -23,7 +23,7 @@ public:
 	virtual void SaveState(const TArray<TWeakObjectPtr<UObject>>& RuntimeObjects, IMovieScenePlayer& Player, FMovieSceneSequenceInstance& SequenceInstance) override;
 	virtual void RestoreState(const TArray<TWeakObjectPtr<UObject>>& RuntimeObjects, IMovieScenePlayer& Player, FMovieSceneSequenceInstance& SequenceInstance) override;
 	virtual void Update(EMovieSceneUpdateData& UpdateData, const TArray<TWeakObjectPtr<UObject>>& RuntimeObjects, IMovieScenePlayer& Player, FMovieSceneSequenceInstance& SequenceInstance) override;
-	virtual void RefreshInstance(const TArray<TWeakObjectPtr<UObject>>& RuntimeObjects, IMovieScenePlayer& Player, FMovieSceneSequenceInstance& SequenceInstance ) override {}
+	virtual void RefreshInstance(const TArray<TWeakObjectPtr<UObject>>& RuntimeObjects, IMovieScenePlayer& Player, FMovieSceneSequenceInstance& SequenceInstance) override;
 	virtual void ClearInstance(IMovieScenePlayer& Player, FMovieSceneSequenceInstance& SequenceInstance ) override {}
 
 private:
diff --git a/Engine/Source/Runtime/Renderer/Private/PostProcess/PostProcessTonemap.cpp b/Engine/Source/Runtime/Renderer/Private/PostProcess/PostProcessTonemap.cpp
--- a/Engine/Source/Runtime/Renderer/Private/PostProcess/PostProcessTonemap.cpp
+++ b/Engine/Source/Runtime/Renderer/Private/PostProcess/PostProcessTonemap.cpp
@@ -27,9 +27,10 @@ static TAutoConsoleVariable<int32> CVarTonemapperGamut(
 	TEXT("r.TonemapperOutputGamut"),
 	0,
 	TEXT("0: use Rec.709/sRGB, D65\n")
-	TEXT("1: use P3 (DCI), D65\n")
+	TEXT("1: use P3, D65\n")
 	TEXT("2: use Rec.2020, D65\n")
-	TEXT("3: use ACES, D60"),
+	TEXT("3: use ACES, D60\n")
+	TEXT("4: use ACEScg, D60"),
 	ECVF_Scalability | ECVF_RenderThreadSafe);
 
 static TAutoConsoleVariable<int32> CVarTonemapper2084(
@@ -992,13 +993,6 @@ public:
 		// The gamut for output
 		static TConsoleVariableData<int32>* CVarOutputGamut = IConsoleManager::Get().FindTConsoleVariableDataInt(TEXT("r.TonemapperOutputGamut"));
 		int32 OutputGamutValue = CVarOutputGamut->GetValueOnRenderThread();
-
-		// Write ACES data when storing HDR frames
-		if (DumpFramesAsHDRValue)
-		{
-			OutputGamutValue = 3;
-		}
-
 		SetShaderValue(Context.RHICmdList, ShaderRHI, OutputGamut, OutputGamutValue);
 
 		SetShaderValue(Context.RHICmdList, ShaderRHI, OverlayColor, Context.View.OverlayColor);
@@ -1161,11 +1155,12 @@ IMPLEMENT_SHADER_TYPE(template<>, TPostProcessTonemapVS<true>, TEXT("PostProcess
 IMPLEMENT_SHADER_TYPE(template<>, TPostProcessTonemapVS<false>, TEXT("PostProcessTonemap"), TEXT("MainVS"), SF_Vertex);
 
 
-FRCPassPostProcessTonemap::FRCPassPostProcessTonemap(const FViewInfo& InView, bool bInDoGammaOnly, bool bInDoEyeAdaptation)
+FRCPassPostProcessTonemap::FRCPassPostProcessTonemap(const FViewInfo& InView, bool bInDoGammaOnly, bool bInDoEyeAdaptation, bool bInHDROutput)
 	: bDoGammaOnly(bInDoGammaOnly)
 	, bDoScreenPercentageInTonemapper(false)
 	, bDoEyeAdaptation(bInDoEyeAdaptation)
-	, View(InView)	
+	, bHDROutput(bInHDROutput)
+	, View(InView)
 {
 	uint32 ConfigBitmask = TonemapperGenerateBitmaskPC(&InView, bDoGammaOnly);
 	ConfigIndexPC = TonemapperFindLeastExpensive(TonemapperConfBitmaskPC, sizeof(TonemapperConfBitmaskPC)/4, TonemapperCostTab, ConfigBitmask);;
@@ -1334,7 +1329,7 @@ FPooledRenderTargetDesc FRCPassPostProcessTonemap::ComputeOutputDesc(EPassOutput
 
 	Ret.Reset();
 	// RGB is the color in LDR, A is the luminance for PostprocessAA
-	Ret.Format = PF_B8G8R8A8;
+	Ret.Format = bHDROutput ? PF_FloatRGBA : PF_B8G8R8A8;
 	Ret.DebugName = TEXT("Tonemap");
 	Ret.ClearValue = FClearValueBinding(FLinearColor(0, 0, 0, 0));
 
diff --git a/Engine/Source/Runtime/Renderer/Private/PostProcess/PostProcessTonemap.h b/Engine/Source/Runtime/Renderer/Private/PostProcess/PostProcessTonemap.h
--- a/Engine/Source/Runtime/Renderer/Private/PostProcess/PostProcessTonemap.h
+++ b/Engine/Source/Runtime/Renderer/Private/PostProcess/PostProcessTonemap.h
@@ -41,7 +41,7 @@ class FRCPassPostProcessTonemap : public TRenderingCompositePassBase<4, 1>
 {
 public:
 	// constructor
-	FRCPassPostProcessTonemap(const FViewInfo& InView, bool bInDoGammaOnly, bool bDoEyeAdaptation);
+	FRCPassPostProcessTonemap(const FViewInfo& InView, bool bInDoGammaOnly, bool bDoEyeAdaptation, bool bHDROutput);
 
 	// interface FRenderingCompositePass ---------
 
@@ -53,6 +53,7 @@ public:
 	bool bDoScreenPercentageInTonemapper;
 private:
 	bool bDoEyeAdaptation;
+	bool bHDROutput;
 	uint32 ConfigIndexPC;
 
 	const FViewInfo& View;
diff --git a/Engine/Source/Runtime/Renderer/Private/PostProcess/PostProcessing.cpp b/Engine/Source/Runtime/Renderer/Private/PostProcess/PostProcessing.cpp
--- a/Engine/Source/Runtime/Renderer/Private/PostProcess/PostProcessing.cpp
+++ b/Engine/Source/Runtime/Renderer/Private/PostProcess/PostProcessing.cpp
@@ -154,6 +154,13 @@ static TAutoConsoleVariable<float> CVarTonemapperMergeThreshold(
 	TEXT("Defauls to 0.49 (e.g., if r.ScreenPercentage is 70 or higher, try to merge)"),
 	ECVF_Scalability | ECVF_RenderThreadSafe);
 
+static TAutoConsoleVariable<int32> CVarMotionBlurNew(
+	TEXT("r.MotionBlurNew"),
+	1,
+	TEXT(""),
+	ECVF_RenderThreadSafe
+);
+
 static TAutoConsoleVariable<int32> CVarMotionBlurScatter(
 	TEXT("r.MotionBlurScatter"),
 	0,
@@ -324,13 +331,14 @@ static FRCPassPostProcessTonemap* AddTonemapper(
 	const FRenderingCompositeOutputRef& BloomOutputCombined,
 	const FRenderingCompositeOutputRef& EyeAdaptation,
 	const EAutoExposureMethod& EyeAdapationMethodId,
-	const bool bDoGammaOnly)
+	const bool bDoGammaOnly,
+	const bool bHDRTonemapperOutput)
 {
 	const FEngineShowFlags& EngineShowFlags = Context.View.Family->EngineShowFlags;
 
 	FRenderingCompositePass* CombinedLUT = Context.Graph.RegisterPass(new(FMemStack::Get()) FRCPassPostProcessCombineLUTs(Context.View.GetShaderPlatform()));
 	const bool bDoEyeAdaptation = IsAutoExposureMethodSupported(Context.View.GetFeatureLevel(), EyeAdapationMethodId);
-	FRCPassPostProcessTonemap* PostProcessTonemap =	Context.Graph.RegisterPass(new(FMemStack::Get()) FRCPassPostProcessTonemap(Context.View, bDoGammaOnly, bDoEyeAdaptation));
+	FRCPassPostProcessTonemap* PostProcessTonemap =	Context.Graph.RegisterPass(new(FMemStack::Get()) FRCPassPostProcessTonemap(Context.View, bDoGammaOnly, bDoEyeAdaptation, bHDRTonemapperOutput));
 
 	PostProcessTonemap->SetInput(ePId_Input0, Context.FinalOutput);
 	PostProcessTonemap->SetInput(ePId_Input1, BloomOutputCombined);
@@ -358,7 +366,7 @@ static void AddSelectionOutline(FPostprocessContext& Context)
 
 static void AddGammaOnlyTonemapper(FPostprocessContext& Context)
 {
-	FRenderingCompositePass* PostProcessTonemap = Context.Graph.RegisterPass(new(FMemStack::Get()) FRCPassPostProcessTonemap(Context.View, true, false/*eye*/));
+	FRenderingCompositePass* PostProcessTonemap = Context.Graph.RegisterPass(new(FMemStack::Get()) FRCPassPostProcessTonemap(Context.View, true, false/*eye*/, false));
 
 	PostProcessTonemap->SetInput(ePId_Input0, Context.FinalOutput);
 
@@ -893,7 +901,7 @@ static bool HasPostProcessMaterial(FPostprocessContext& Context, EBlendableLocat
 	return false;
 }
 
-static void AddPostProcessMaterial(FPostprocessContext& Context, EBlendableLocation InLocation, FRenderingCompositeOutputRef SeparateTranslucency, FRenderingCompositeOutputRef HDRColor = FRenderingCompositeOutputRef())
+static void AddPostProcessMaterial(FPostprocessContext& Context, EBlendableLocation InLocation, FRenderingCompositeOutputRef SeparateTranslucency, FRenderingCompositeOutputRef PostTonemapHDRColor = FRenderingCompositeOutputRef())
 {
 	if( !Context.View.Family->EngineShowFlags.PostProcessing ||
 		!Context.View.Family->EngineShowFlags.PostProcessMaterial ||
@@ -968,7 +976,7 @@ static void AddPostProcessMaterial(FPostprocessContext& Context, EBlendableLocat
 		// This input is only needed for visualization and frame dumping
 		if (bVisualizingBuffer)
 		{
-			Node->SetInput(ePId_Input2, HDRColor);
+			Node->SetInput(ePId_Input2, PostTonemapHDRColor);
 		}
 
 		Context.FinalOutput = FRenderingCompositeOutputRef(Node);
@@ -1019,7 +1027,7 @@ static void AddHighResScreenshotMask(FPostprocessContext& Context, FRenderingCom
 	}
 }
 
-static void AddGBufferVisualizationOverview(FPostprocessContext& Context, FRenderingCompositeOutputRef& SeparateTranslucencyInput, FRenderingCompositeOutputRef& HDRColorInput)
+static void AddGBufferVisualizationOverview(FPostprocessContext& Context, FRenderingCompositeOutputRef& SeparateTranslucencyInput, FRenderingCompositeOutputRef& PostTonemapHDRColorInput)
 {
 	static const auto CVarDumpFrames = IConsoleManager::Get().FindTConsoleVariableDataInt(TEXT("r.BufferVisualizationDumpFrames"));
 	static const auto CVarDumpFramesAsHDR = IConsoleManager::Get().FindTConsoleVariableDataInt(TEXT("r.BufferVisualizationDumpFramesAsHDR"));
@@ -1057,7 +1065,7 @@ static void AddGBufferVisualizationOverview(FPostprocessContext& Context, FRende
 					FRenderingCompositePass* MaterialPass = Context.Graph.RegisterPass(new(FMemStack::Get()) FRCPassPostProcessMaterial(*It, Context.View.GetFeatureLevel(), OutputFormat));
 					MaterialPass->SetInput(ePId_Input0, FRenderingCompositeOutputRef(IncomingStage));
 					MaterialPass->SetInput(ePId_Input1, FRenderingCompositeOutputRef(SeparateTranslucencyInput));
-					MaterialPass->SetInput(ePId_Input2, FRenderingCompositeOutputRef(HDRColorInput));
+					MaterialPass->SetInput(ePId_Input2, FRenderingCompositeOutputRef(PostTonemapHDRColorInput));
 
 					auto Proxy = MaterialInterface->GetRenderProxy(false);
 					const FMaterial* Material = Proxy->GetMaterial(Context.View.GetFeatureLevel());
@@ -1176,6 +1184,8 @@ void FPostProcessing::Process(FRHICommandListImmediate& RHICmdList, const FViewI
 		FRenderingCompositeOutputRef HistogramOverScreen;
 		// not always valid
 		FRenderingCompositeOutputRef Histogram;
+		// not always valid
+		FRenderingCompositeOutputRef PostTonemapHDRColor;
 
 		class FAutoExposure
 		{
@@ -1241,7 +1251,10 @@ void FPostProcessing::Process(FRHICommandListImmediate& RHICmdList, const FViewI
 		{
 			bAllowTonemapper = false;
 		}
-	
+
+		static const auto CVarDumpFramesAsHDR = IConsoleManager::Get().FindTConsoleVariableDataInt(TEXT("r.BufferVisualizationDumpFramesAsHDR"));
+		const bool bHDRTonemapperOutput = bAllowTonemapper && (GetHighResScreenshotConfig().bCaptureHDR || CVarDumpFramesAsHDR->GetValueOnRenderThread());
+
 		FRCPassPostProcessTonemap* Tonemapper = 0;
 
 		// add the passes we want to add to the graph (commenting a line means the pass is not inserted into the graph) ---------
@@ -1644,7 +1657,17 @@ void FPostProcessing::Process(FRHICommandListImmediate& RHICmdList, const FViewI
 				}
 				else
 				{
-					Tonemapper = AddTonemapper(Context, BloomOutputCombined, AutoExposure.EyeAdaptation, AutoExposure.MethodId, false);
+					Tonemapper = AddTonemapper(Context, BloomOutputCombined, AutoExposure.EyeAdaptation, AutoExposure.MethodId, false, bHDRTonemapperOutput);
+				}
+
+				PostTonemapHDRColor = Context.FinalOutput;
+
+				// Add a pass-through as tonemapper will be forced LDR if final pass in chain 
+				if (bHDRTonemapperOutput)
+				{
+					FRenderingCompositePass* PassthroughNode = Context.Graph.RegisterPass(new(FMemStack::Get()) FRCPassPostProcessPassThrough(nullptr));
+					PassthroughNode->SetInput(ePId_Input0, FRenderingCompositeOutputRef(Context.FinalOutput));
+					Context.FinalOutput = FRenderingCompositeOutputRef(PassthroughNode);
 				}
 			}
 
@@ -1765,7 +1788,7 @@ void FPostProcessing::Process(FRHICommandListImmediate& RHICmdList, const FViewI
 			Context.FinalOutput = FRenderingCompositeOutputRef(Node);
 		}
 		
-		AddPostProcessMaterial(Context, BL_AfterTonemapping, SeparateTranslucency, HDRColor);
+		AddPostProcessMaterial(Context, BL_AfterTonemapping, SeparateTranslucency, PostTonemapHDRColor);
 
 #if WITH_EDITOR
 		//Inspect the Final color, GBuffer and HDR
@@ -1795,7 +1818,7 @@ void FPostProcessing::Process(FRHICommandListImmediate& RHICmdList, const FViewI
 			Context.FinalOutput = FRenderingCompositeOutputRef(PassVisualize);
 		}
 
-		AddGBufferVisualizationOverview(Context, SeparateTranslucency, HDRColor);
+		AddGBufferVisualizationOverview(Context, SeparateTranslucency, PostTonemapHDRColor);
 
 		if (bStereoRenderingAndHMD)
 		{
@@ -2273,7 +2296,7 @@ void FPostProcessing::ProcessES2(FRHICommandListImmediate& RHICmdList, const FVi
 		if ( bUseTonemapperFilm)
 		{
 			//@todo Ronin Set to EAutoExposureMethod::AEM_Basic for PC vk crash.
-			AddTonemapper(Context, BloomOutput, nullptr, EAutoExposureMethod::AEM_Histogram, false);
+			AddTonemapper(Context, BloomOutput, nullptr, EAutoExposureMethod::AEM_Histogram, false, false);
 		}
 		else
 		{
diff --git a/Engine/Source/Runtime/Slate/Private/Widgets/Colors/SColorWheel.cpp b/Engine/Source/Runtime/Slate/Private/Widgets/Colors/SColorWheel.cpp
--- a/Engine/Source/Runtime/Slate/Private/Widgets/Colors/SColorWheel.cpp
+++ b/Engine/Source/Runtime/Slate/Private/Widgets/Colors/SColorWheel.cpp
@@ -148,7 +148,5 @@ bool SColorWheel::ProcessMouseAction(const FGeometry& MyGeometry, const FPointer
 		OnValueChanged.ExecuteIfBound(NewColor);
 	}
 
-	GLog->Logf(TEXT("%f"), RelativeRadius);
-
 	return (RelativeRadius <= 1.0f);
 }
