Copying //UE4/Dev-Sequencer to //UE4/Dev-Main (Source: //UE4/Dev-Sequencer @ 3048536)

diff --git a/Engine/Source/Runtime/MovieSceneCapture/Private/CompositionGraphCaptureProtocol.cpp b/Engine/Source/Runtime/MovieSceneCapture/Private/CompositionGraphCaptureProtocol.cpp
--- a/Engine/Source/Runtime/MovieSceneCapture/Private/CompositionGraphCaptureProtocol.cpp
+++ b/Engine/Source/Runtime/MovieSceneCapture/Private/CompositionGraphCaptureProtocol.cpp
@@ -10,14 +10,20 @@
 
 struct FSceneViewExtension : ISceneViewExtension
 {
-	FSceneViewExtension(const TArray<FString>& InRenderPasses, bool bInCaptureFramesInHDR, int32 InCaptureGamut, UMaterialInterface* InPostProcessingMaterial)
-		: RenderPasses(InRenderPasses), bNeedsCapture(true), CaptureGamut(InCaptureGamut), RestoreDumpHDR(0), RestoreDumpGamut(HCGM_Rec709)
+	FSceneViewExtension(const TArray<FString>& InRenderPasses, bool bInCaptureFramesInHDR, int32 InHDRCompressionQuality, int32 InCaptureGamut, UMaterialInterface* InPostProcessingMaterial)
+		: RenderPasses(InRenderPasses)
+		, bNeedsCapture(true)
+		, bCaptureFramesInHDR(bInCaptureFramesInHDR)
+		, HDRCompressionQuality(InHDRCompressionQuality)
+		, CaptureGamut(InCaptureGamut)
+		, PostProcessingMaterial(InPostProcessingMaterial)
+		, RestoreDumpHDR(0)
+		, RestoreHDRCompressionQuality(0)
+		, RestoreDumpGamut(HCGM_Rec709)
 	{
-		PostProcessingMaterial = InPostProcessingMaterial;
-		bCaptureFramesInHDR = bInCaptureFramesInHDR;
-
 		CVarDumpFrames = IConsoleManager::Get().FindConsoleVariable(TEXT("r.BufferVisualizationDumpFrames"));
 		CVarDumpFramesAsHDR = IConsoleManager::Get().FindConsoleVariable(TEXT("r.BufferVisualizationDumpFramesAsHDR"));
+		CVarHDRCompressionQuality = IConsoleManager::Get().FindConsoleVariable(TEXT("r.SaveEXR.CompressionQuality"));
 		CVarDumpGamut = IConsoleManager::Get().FindConsoleVariable(TEXT("r.TonemapperOutputGamut"));
 
 		Disable();
@@ -38,10 +44,13 @@ struct FSceneViewExtension : ISceneViewExtension
 		OutputFilename = MoveTemp(InFilename);
 
 		bNeedsCapture = true;
+
 		RestoreDumpHDR = CVarDumpFramesAsHDR->GetInt();
+		RestoreHDRCompressionQuality = CVarHDRCompressionQuality->GetInt();
 		RestoreDumpGamut = CVarDumpGamut->GetInt();
 
 		CVarDumpFramesAsHDR->Set(bCaptureFramesInHDR);
+		CVarHDRCompressionQuality->Set(HDRCompressionQuality);
 		CVarDumpGamut->Set(CaptureGamut);
 		CVarDumpFrames->Set(1);
 	}
@@ -54,8 +63,10 @@ struct FSceneViewExtension : ISceneViewExtension
 			if (bFinalize)
 			{
 				RestoreDumpHDR = 0;
+				RestoreHDRCompressionQuality = 0;
 			}
 			CVarDumpFramesAsHDR->Set(RestoreDumpHDR);
+			CVarHDRCompressionQuality->Set(RestoreHDRCompressionQuality);
 			CVarDumpGamut->Set(RestoreDumpGamut);
 			CVarDumpFrames->Set(0);
 		}
@@ -112,19 +123,22 @@ struct FSceneViewExtension : ISceneViewExtension
 private:
 	const TArray<FString>& RenderPasses;
 
+	bool bNeedsCapture;
+	FString OutputFilename;
+
 	bool bCaptureFramesInHDR;
+	int32 HDRCompressionQuality;
+	int32 CaptureGamut;
 
 	UMaterialInterface* PostProcessingMaterial;
 
-	bool bNeedsCapture;
-	int32 CaptureGamut;
-	FString OutputFilename;
-
 	IConsoleVariable* CVarDumpFrames;
 	IConsoleVariable* CVarDumpFramesAsHDR;
+	IConsoleVariable* CVarHDRCompressionQuality;
 	IConsoleVariable* CVarDumpGamut;
 
 	int32 RestoreDumpHDR;
+	int32 RestoreHDRCompressionQuality;
 	int32 RestoreDumpGamut;
 };
 
@@ -175,6 +189,7 @@ bool FCompositionGraphCaptureProtocol::Initialize(const FCaptureProtocolInitSett
 	SceneViewport = InSettings.SceneViewport;
 
 	bool bCaptureFramesInHDR = false;
+	int32 HDRCompressionQuality = 0;
 	int32 CaptureGamut = HCGM_Rec709;
 
 	UMaterialInterface* PostProcessingMaterial = nullptr;
@@ -183,6 +198,7 @@ bool FCompositionGraphCaptureProtocol::Initialize(const FCaptureProtocolInitSett
 	{
 		RenderPasses = ProtocolSettings->IncludeRenderPasses.Value;
 		bCaptureFramesInHDR = ProtocolSettings->bCaptureFramesInHDR;
+		HDRCompressionQuality = ProtocolSettings->HDRCompressionQuality;
 		CaptureGamut = ProtocolSettings->CaptureGamut;
 		PostProcessingMaterial = Cast<UMaterialInterface>(ProtocolSettings->PostProcessingMaterial.TryLoad());
 
@@ -198,6 +214,12 @@ bool FCompositionGraphCaptureProtocol::Initialize(const FCaptureProtocolInitSett
 			bCaptureFramesInHDR = bOverrideCaptureFramesInHDR;
 		}
 
+		int32 OverrideHDRCompressionQuality;
+		if( FParse::Value( FCommandLine::Get(), TEXT( "-HDRCompressionQuality=" ), OverrideHDRCompressionQuality ) )
+		{
+			HDRCompressionQuality = OverrideHDRCompressionQuality;
+		}
+
 		int32 OverrideCaptureGamut;
 		if (FParse::Value(FCommandLine::Get(), TEXT("-CaptureGamut="), OverrideCaptureGamut))
 		{
@@ -205,7 +227,7 @@ bool FCompositionGraphCaptureProtocol::Initialize(const FCaptureProtocolInitSett
 		}
 	}
 
-	ViewExtension = MakeShareable(new FSceneViewExtension(RenderPasses, bCaptureFramesInHDR, CaptureGamut, PostProcessingMaterial));
+	ViewExtension = MakeShareable(new FSceneViewExtension(RenderPasses, bCaptureFramesInHDR, HDRCompressionQuality, CaptureGamut, PostProcessingMaterial));
 
 	return true;
 }
diff --git a/Engine/Source/Runtime/MovieSceneCapture/Private/MovieSceneCapture.cpp b/Engine/Source/Runtime/MovieSceneCapture/Private/MovieSceneCapture.cpp
--- a/Engine/Source/Runtime/MovieSceneCapture/Private/MovieSceneCapture.cpp
+++ b/Engine/Source/Runtime/MovieSceneCapture/Private/MovieSceneCapture.cpp
@@ -8,6 +8,7 @@
 #include "SceneViewExtension.h"
 #include "JsonObjectConverter.h"
 #include "RemoteConfigIni.h"
+#include "SceneViewport.h"
 
 #if WITH_EDITOR
 #include "ImageWrapper.h"
@@ -58,8 +59,8 @@ UMovieSceneCapture::UMovieSceneCapture(const FObjectInitializer& Initializer)
 		InheritedCommandLineArguments.Append(Switch);
 		InheritedCommandLineArguments.AppendChar(' ');
 	}
-
-	AdditionalCommandLineArguments += TEXT("-NoLoadingScreen -NOSCREENMESSAGES -ForceRes");
+	
+	AdditionalCommandLineArguments += TEXT("-NOSCREENMESSAGES");
 
 	Handle = FUniqueMovieSceneCaptureHandle();
 
@@ -158,7 +159,7 @@ void UMovieSceneCapture::Initialize(TSharedPtr<FSceneViewport> InSceneViewport,
 	FormatMappings.Add(TEXT("fps"), FString::Printf(TEXT("%d"), Settings.FrameRate));
 	FormatMappings.Add(TEXT("width"), FString::Printf(TEXT("%d"), CachedMetrics.Width));
 	FormatMappings.Add(TEXT("height"), FString::Printf(TEXT("%d"), CachedMetrics.Height));
-	FormatMappings.Add(TEXT("world"), GWorld->GetName());
+	FormatMappings.Add(TEXT("world"), InSceneViewport->GetClient()->GetWorld()->GetName());
 
 	FrameNumberFormat = FString::Printf(TEXT("%%0%dd"), Settings.ZeroPadFrameNumbers);
 
diff --git a/Engine/Source/Runtime/MovieSceneCapture/Private/MovieSceneCaptureHelpers.cpp b/Engine/Source/Runtime/MovieSceneCapture/Private/MovieSceneCaptureHelpers.cpp
--- a/Engine/Source/Runtime/MovieSceneCapture/Private/MovieSceneCaptureHelpers.cpp
+++ b/Engine/Source/Runtime/MovieSceneCapture/Private/MovieSceneCaptureHelpers.cpp
@@ -56,7 +56,7 @@ struct FShotData
 	float EditOutTime;
 };
 
-bool SMPTEToTime(FString SMPTE, float FrameRate, float& OutTime)
+void SMPTEToTime(FString SMPTE, float FrameRate, float& OutTime)
 {
 	TArray<FString> OutArray;
 	SMPTE.ParseIntoArray(OutArray, TEXT(":"));
@@ -69,16 +69,13 @@ bool SMPTEToTime(FString SMPTE, float FrameRate, float& OutTime)
 		float Frames = FCString::Atof(*OutArray[3]);
 
 		OutTime = Hours * 3600.f + Minutes * 60.f  + Seconds + (Frames / FrameRate);
-		return true;
 	}
 	// The edl is in frames
 	else
 	{
 		float Frames = FCString::Atof(*SMPTE);
-		OutTime = Frames * FrameRate;
+		OutTime = Frames / FrameRate;
 	}
-
-	return false;
 }
 
 FString TimeToSMPTE(float InTime, float FrameRate)
@@ -163,12 +160,13 @@ void ParseFromEDL(const FString& InputString, float FrameRate, TArray<FShotData>
 
 				// If everything checks out
 				if (TrackType != FShotData::ETrackType::TT_None &&
-					EditType != FShotData::EEditType::ET_None &&
-					SMPTEToTime(InputChars[4], FrameRate, SourceInTime) &&
-					SMPTEToTime(InputChars[5], FrameRate, SourceOutTime) &&
-					SMPTEToTime(InputChars[6], FrameRate, EditInTime) &&
-					SMPTEToTime(InputChars[7], FrameRate, EditOutTime))
+					EditType != FShotData::EEditType::ET_None)
 				{
+					SMPTEToTime(InputChars[4], FrameRate, SourceInTime);
+					SMPTEToTime(InputChars[5], FrameRate, SourceOutTime);
+					SMPTEToTime(InputChars[6], FrameRate, EditInTime);
+					SMPTEToTime(InputChars[7], FrameRate, EditOutTime);
+
 					bFoundEventLine = true;
 					continue; // Go to the next line
 				}
@@ -343,16 +341,23 @@ bool MovieSceneCaptureHelpers::ImportEDL(UMovieScene* InMovieScene, float InFram
 			for (UMovieSceneSection* Section : CinematicShotTrack->GetAllSections())
 			{
 				UMovieSceneCinematicShotSection* CinematicShotSection = Cast<UMovieSceneCinematicShotSection>(Section);
-				if (CinematicShotSection && CinematicShotSection->GetSequence()->GetName() == ShotName)
+				if (CinematicShotSection != nullptr)
 				{
-					ShotSection = CinematicShotSection;
-					break;
+					UMovieSceneSequence* ShotSequence = CinematicShotSection->GetSequence();
+				
+					if (ShotSequence != nullptr && ShotSequence->GetName() == ShotName)
+					{
+						ShotSection = CinematicShotSection;
+						break;
+					}
 				}
 			}
 
 			// If the shot doesn't already exist, create it
 			if (!ShotSection)
 			{
+				UMovieSceneSequence* SequenceToAdd = nullptr;
+
 				FAssetRegistryModule& AssetRegistryModule = FModuleManager::LoadModuleChecked<FAssetRegistryModule>(TEXT("AssetRegistry"));
 
 				// Collect a full list of assets with the specified class
@@ -363,11 +368,13 @@ bool MovieSceneCaptureHelpers::ImportEDL(UMovieScene* InMovieScene, float InFram
 				{
 					if (AssetData.AssetName == *ShotName)
 					{	
-						CinematicShotTrack->Modify();
-						ShotSection = Cast<UMovieSceneCinematicShotSection>(CinematicShotTrack->AddSequence(Cast<ULevelSequence>(AssetData.GetAsset()), ShotData.EditInTime, ShotData.EditOutTime - ShotData.EditInTime));
+						SequenceToAdd = Cast<ULevelSequence>(AssetData.GetAsset());
 						break;
 					}
 				}
+
+				CinematicShotTrack->Modify();
+				ShotSection = Cast<UMovieSceneCinematicShotSection>(CinematicShotTrack->AddSequence(SequenceToAdd, ShotData.EditInTime, ShotData.EditOutTime - ShotData.EditInTime));
 			}
 
 			// Conform this shot section
@@ -384,18 +391,16 @@ bool MovieSceneCaptureHelpers::ImportEDL(UMovieScene* InMovieScene, float InFram
 	return true;
 }
 
-bool MovieSceneCaptureHelpers::ExportEDL(const UMovieScene* InMovieScene, float InFrameRate, FString InSaveDirectory)
+bool MovieSceneCaptureHelpers::ExportEDL(const UMovieScene* InMovieScene, float InFrameRate, FString InSaveFilename)
 {
 	FString SequenceName = InMovieScene->GetOuter()->GetName();
-	FString SequencePath = InMovieScene->GetOuter()->GetPathName();
-	SequencePath = FPaths::GetPath(SequencePath);
+	FString SaveBasename = FPaths::GetPath(InSaveFilename) / FPaths::GetBaseFilename(InSaveFilename);
 
-	// Pop open a dialog to request the location of the edl (.edl, or .rv formats accepted)
 	TArray<FString> SaveFilenames;
-	if (!InSaveDirectory.IsEmpty())
+	if (!SaveBasename.IsEmpty())
 	{
-		SaveFilenames.Add(InSaveDirectory / SequenceName + TEXT(".rv"));
-		SaveFilenames.Add(InSaveDirectory / SequenceName + TEXT(".edl"));
+		SaveFilenames.Add(SaveBasename + TEXT(".rv"));
+		SaveFilenames.Add(SaveBasename + TEXT(".edl"));
 	}
 
 	TArray<FShotData> ShotDataArray;
@@ -410,7 +415,14 @@ bool MovieSceneCaptureHelpers::ExportEDL(const UMovieScene* InMovieScene, float
 			for (auto ShotSection : CinematicShotTrack->GetAllSections())
 			{
 				UMovieSceneCinematicShotSection* CinematicShotSection = Cast<UMovieSceneCinematicShotSection>(ShotSection);
-				
+
+				if ( CinematicShotSection->GetSequence() == nullptr )
+				{
+					// If the shot doesn't have a valid sequence skip it.  This is currently the case for filler sections.
+					// TODO: Handle this properly in the edl output.
+					continue;
+				}
+
 				FString ShotName = CinematicShotSection->GetShotDisplayName().ToString();
 				FString ShotPath = CinematicShotSection->GetSequence()->GetMovieScene()->GetOuter()->GetPathName();
 
diff --git a/Engine/Source/Runtime/MovieSceneCapture/Private/MovieSceneCaptureModule.cpp b/Engine/Source/Runtime/MovieSceneCapture/Private/MovieSceneCaptureModule.cpp
--- a/Engine/Source/Runtime/MovieSceneCapture/Private/MovieSceneCaptureModule.cpp
+++ b/Engine/Source/Runtime/MovieSceneCapture/Private/MovieSceneCaptureModule.cpp
@@ -86,8 +86,16 @@ private:
 		UGameEngine* GameEngine = Cast<UGameEngine>(GEngine);
 		if (GameEngine && StartupMovieCaptureHandle.IsValid())
 		{
-			IMovieSceneCaptureInterface* StartupCaptureInterface = RetrieveMovieSceneInterface(StartupMovieCaptureHandle);
-			StartupCaptureInterface->Initialize(GameEngine->SceneViewport.ToSharedRef());
+			if (!GameEngine->SceneViewport->GetClient()->GetWorld())
+			{
+				// @todo: Set exit code to EMovieSceneCaptureExitCode::WorldNotFound when we have the ability to do so
+				FPlatformMisc::RequestExit(false);
+			}
+			else
+			{
+				IMovieSceneCaptureInterface* StartupCaptureInterface = RetrieveMovieSceneInterface(StartupMovieCaptureHandle);
+				StartupCaptureInterface->Initialize(GameEngine->SceneViewport.ToSharedRef());
+			}
 		}
 
 		StartupMovieCaptureHandle = FMovieSceneCaptureHandle();
diff --git a/Engine/Source/Runtime/MovieSceneCapture/Public/ErrorCodes.h b/Engine/Source/Runtime/MovieSceneCapture/Public/ErrorCodes.h
--- a/Engine/Source/Runtime/MovieSceneCapture/Public/ErrorCodes.h
+++ b/Engine/Source/Runtime/MovieSceneCapture/Public/ErrorCodes.h
@@ -3,8 +3,9 @@
 #pragma once
 
 /** Process exit codes */
-struct FMovieSceneCaptureExitCodes
+enum class EMovieSceneCaptureExitCode : uint32
 {
-	static const int32 Base = 0xAC701000;
-	static const int32 AssetNotFound = Base + 1;
+	Base = 0xAC701000,
+	AssetNotFound,
+	WorldNotFound,
 };
\ No newline at end of file
diff --git a/Engine/Source/Runtime/MovieSceneCapture/Public/MovieSceneCaptureHelpers.h b/Engine/Source/Runtime/MovieSceneCapture/Public/MovieSceneCaptureHelpers.h
--- a/Engine/Source/Runtime/MovieSceneCapture/Public/MovieSceneCaptureHelpers.h
+++ b/Engine/Source/Runtime/MovieSceneCapture/Public/MovieSceneCaptureHelpers.h
@@ -23,8 +23,8 @@ public:
 	 *
 	 * @param InMovieScene The movie scene with the cinematic shot track and audio tracks to export
 	 * @param InFrameRate The frame rate to export the EDL at
-	 * @param InSaveDirectory The directory path to save to.
+	 * @param InSaveFilename The file path to save to.
 	 * @return Whether the export was successful
 	 */
-	static bool ExportEDL(const UMovieScene* InMovieScene, float InFrameRate, FString InSaveDirectory);
+	static bool ExportEDL(const UMovieScene* InMovieScene, float InFrameRate, FString InSaveFilename);
 };
\ No newline at end of file
diff --git a/Engine/Source/Runtime/MovieSceneCapture/Public/MovieSceneCaptureSettings.h b/Engine/Source/Runtime/MovieSceneCapture/Public/MovieSceneCaptureSettings.h
--- a/Engine/Source/Runtime/MovieSceneCapture/Public/MovieSceneCaptureSettings.h
+++ b/Engine/Source/Runtime/MovieSceneCapture/Public/MovieSceneCaptureSettings.h
@@ -49,6 +49,7 @@ struct MOVIESCENECAPTURE_API FMovieSceneCaptureSettings
 	 * {world}		- The name of the current world
 	 * {quality}	- The image compression quality setting
 	 * {material}   - The material/render pass
+	 * {shot}       - The name of the level sequence asset shot being played
 	 */
 	UPROPERTY(config, EditAnywhere, Category=General, DisplayName="Filename Format")
 	FString OutputFormat;
diff --git a/Engine/Source/Runtime/MovieSceneCapture/Public/Protocols/CompositionGraphCaptureProtocol.h b/Engine/Source/Runtime/MovieSceneCapture/Public/Protocols/CompositionGraphCaptureProtocol.h
--- a/Engine/Source/Runtime/MovieSceneCapture/Public/Protocols/CompositionGraphCaptureProtocol.h
+++ b/Engine/Source/Runtime/MovieSceneCapture/Public/Protocols/CompositionGraphCaptureProtocol.h
@@ -45,6 +45,10 @@ public:
 	UPROPERTY(config, EditAnywhere, Category="Composition Graph Options")
 	bool bCaptureFramesInHDR;
 
+	/** Compression Quality for HDR Frames (0 for no compression, 1 for default compression which can be slow) */
+	UPROPERTY(config, EditAnywhere, Category="Composition Graph Options", meta = (EditCondition = "bCaptureFramesInHDR"))
+	int32 HDRCompressionQuality;
+
 	/** The color gamut to use when storing HDR captured data. The gamut depends on whether the bCaptureFramesInHDR option is enabled. */
 	UPROPERTY(config, EditAnywhere, Category="Composition Graph Options", meta = (EditCondition = "bCaptureFramesInHDR"))
 	TEnumAsByte<enum EHDRCaptureGamut> CaptureGamut;
diff --git a/Engine/Source/Runtime/MovieSceneTracks/Private/Sections/MovieScene3DTransformSection.cpp b/Engine/Source/Runtime/MovieSceneTracks/Private/Sections/MovieScene3DTransformSection.cpp
--- a/Engine/Source/Runtime/MovieSceneTracks/Private/Sections/MovieScene3DTransformSection.cpp
+++ b/Engine/Source/Runtime/MovieSceneTracks/Private/Sections/MovieScene3DTransformSection.cpp
@@ -92,6 +92,9 @@ void FMovieScene3DTransformKeyStruct::PropagateChanges(const FPropertyChangedEve
 
 UMovieScene3DTransformSection::UMovieScene3DTransformSection(const FObjectInitializer& ObjectInitializer)
 	: Super(ObjectInitializer)
+#if WITH_EDITORONLY_DATA
+	, Show3DTrajectory(EShow3DTrajectory::EST_OnlyWhenSelected)
+#endif
 { }
 
 
diff --git a/Engine/Source/Runtime/MovieSceneTracks/Private/Sections/MovieSceneSubSection.cpp b/Engine/Source/Runtime/MovieSceneTracks/Private/Sections/MovieSceneSubSection.cpp
--- a/Engine/Source/Runtime/MovieSceneTracks/Private/Sections/MovieSceneSubSection.cpp
+++ b/Engine/Source/Runtime/MovieSceneTracks/Private/Sections/MovieSceneSubSection.cpp
@@ -75,16 +75,15 @@ bool UMovieSceneSubSection::IsSetAsRecording()
 	return GetRecordingSection() != nullptr;
 }
 
-const FString& UMovieSceneSubSection::GetActorToRecord()
+AActor* UMovieSceneSubSection::GetActorToRecord()
 {
 	UMovieSceneSubSection* RecordingSection = GetRecordingSection();
 	if(RecordingSection)
 	{
-		return RecordingSection->NameOfActorToRecord;
+		return RecordingSection->ActorToRecord.Get();
 	}
 
-	static FString EmptyString;
-	return EmptyString;
+	return nullptr;
 }
 
 #if WITH_EDITOR
diff --git a/Engine/Source/Runtime/MovieSceneTracks/Private/TrackInstances/MovieSceneAudioTrackInstance.cpp b/Engine/Source/Runtime/MovieSceneTracks/Private/TrackInstances/MovieSceneAudioTrackInstance.cpp
--- a/Engine/Source/Runtime/MovieSceneTracks/Private/TrackInstances/MovieSceneAudioTrackInstance.cpp
+++ b/Engine/Source/Runtime/MovieSceneTracks/Private/TrackInstances/MovieSceneAudioTrackInstance.cpp
@@ -39,8 +39,7 @@ void FMovieSceneAudioTrackInstance::Update(EMovieSceneUpdateData& UpdateData, co
 
 	if (Player.GetPlaybackStatus() == EMovieScenePlayerStatus::Playing)
 	{
-		float LastPosition = UpdateData.bLooped ? UpdateData.Position : UpdateData.LastPosition;
-		if (UpdateData.Position >= LastPosition)
+		if (UpdateData.Position >= UpdateData.LastPosition)
 		{
 			TMap<int32, TArray<UMovieSceneAudioSection*> > AudioSectionsBySectionIndex;
 			for (int32 i = 0; i < AudioSections.Num(); ++i)
@@ -70,7 +69,7 @@ void FMovieSceneAudioTrackInstance::Update(EMovieSceneUpdateData& UpdateData, co
 							UMovieSceneAudioSection* AudioSection = MovieSceneAudioSections[i];
 							if (AudioSection->IsTimeWithinAudioRange(UpdateData.Position))
 							{
-								if (!AudioSection->IsTimeWithinAudioRange(LastPosition) || !Component->IsPlaying())
+								if (!AudioSection->IsTimeWithinAudioRange(UpdateData.LastPosition) || !Component->IsPlaying())
 								{
 									PlaySound(AudioSection, Component, UpdateData.Position);
 								}
diff --git a/Engine/Source/Runtime/MovieSceneTracks/Private/TrackInstances/MovieSceneCameraAnimTrackInstance.cpp b/Engine/Source/Runtime/MovieSceneTracks/Private/TrackInstances/MovieSceneCameraAnimTrackInstance.cpp
--- a/Engine/Source/Runtime/MovieSceneTracks/Private/TrackInstances/MovieSceneCameraAnimTrackInstance.cpp
+++ b/Engine/Source/Runtime/MovieSceneTracks/Private/TrackInstances/MovieSceneCameraAnimTrackInstance.cpp
@@ -89,7 +89,7 @@ void FMovieSceneCameraAnimTrackInstance::Update(EMovieSceneUpdateData& UpdateDat
 							// make it root so GC doesn't take it away
 							CamAnimInst->AddToRoot();
 							CamAnimInst->SetStopAutomatically(false);
-							CamAnimInst->Play(ActiveSection->GetCameraAnim(), GetTempCameraActor(), ActiveSection->PlayRate, ActiveSection->PlayScale, ActiveSection->BlendInTime, ActiveSection->BlendOutTime, ActiveSection->bLooping, ActiveSection->bRandomStartTime);
+							CamAnimInst->Play(ActiveSection->GetCameraAnim(), GetTempCameraActor(Player), ActiveSection->PlayRate, ActiveSection->PlayScale, ActiveSection->BlendInTime, ActiveSection->BlendOutTime, ActiveSection->bLooping, ActiveSection->bRandomStartTime);
 						}
 
 						SectionInstanceData.CameraAnimInst = CamAnimInst;
@@ -99,7 +99,7 @@ void FMovieSceneCameraAnimTrackInstance::Update(EMovieSceneUpdateData& UpdateDat
 					{
 						// prepare temp camera actor by resetting it
 						{
-							ACameraActor* const CameraActor = GetTempCameraActor();
+							ACameraActor* const CameraActor = GetTempCameraActor(Player);
 							CameraActor->SetActorLocationAndRotation(FVector::ZeroVector, FRotator::ZeroRotator);
 
 							ACameraActor const* const DefaultCamActor = GetDefault<ACameraActor>();
@@ -139,7 +139,7 @@ void FMovieSceneCameraAnimTrackInstance::Update(EMovieSceneUpdateData& UpdateDat
 
 							// harvest PP changes
 							{
-								UCameraComponent* AnimCamComp = GetTempCameraActor()->GetCameraComponent();
+								UCameraComponent* AnimCamComp = GetTempCameraActor(Player)->GetCameraComponent();
 								if (AnimCamComp)
 								{
 									SectionInstanceData.PostProcessingBlendWeight = AnimCamComp->PostProcessBlendWeight;
@@ -175,7 +175,7 @@ void FMovieSceneCameraAnimTrackInstance::Update(EMovieSceneUpdateData& UpdateDat
 	}
 }
 
-ACameraActor* FMovieSceneCameraAnimTrackInstance::GetTempCameraActor()
+ACameraActor* FMovieSceneCameraAnimTrackInstance::GetTempCameraActor(IMovieScenePlayer& InMovieScenePlayer)
 {
 	if (TempCameraActor.IsValid() == false)
 	{
@@ -183,7 +183,7 @@ ACameraActor* FMovieSceneCameraAnimTrackInstance::GetTempCameraActor()
 		FActorSpawnParameters SpawnInfo;
 		SpawnInfo.SpawnCollisionHandlingOverride = ESpawnActorCollisionHandlingMethod::AlwaysSpawn;
 		SpawnInfo.ObjectFlags |= RF_Transient;	// We never want to save these temp actors into a map
-		ACameraActor* const Cam = GWorld->SpawnActor<ACameraActor>(SpawnInfo);		// #fixme, GWorld!
+		ACameraActor* const Cam = InMovieScenePlayer.GetPlaybackContext()->GetWorld()->SpawnActor<ACameraActor>(SpawnInfo);
 		if (Cam)
 		{
 #if WITH_EDITOR
diff --git a/Engine/Source/Runtime/MovieSceneTracks/Private/TrackInstances/MovieSceneCameraCutTrackInstance.cpp b/Engine/Source/Runtime/MovieSceneTracks/Private/TrackInstances/MovieSceneCameraCutTrackInstance.cpp
--- a/Engine/Source/Runtime/MovieSceneTracks/Private/TrackInstances/MovieSceneCameraCutTrackInstance.cpp
+++ b/Engine/Source/Runtime/MovieSceneTracks/Private/TrackInstances/MovieSceneCameraCutTrackInstance.cpp
@@ -49,7 +49,7 @@ void FMovieSceneCameraCutTrackInstance::RestoreState(const TArray<TWeakObjectPtr
 void FMovieSceneCameraCutTrackInstance::Update(EMovieSceneUpdateData& UpdateData, const TArray<TWeakObjectPtr<UObject>>& RuntimeObjects, class IMovieScenePlayer& Player, FMovieSceneSequenceInstance& SequenceInstance) 
 {
 	// Don't switch CameraCuts when in preroll
-	if (UpdateData.bPreroll)
+	if (UpdateData.bPreroll || !UpdateData.bUpdateCameras)
 	{
 		return;
 	}
diff --git a/Engine/Source/Runtime/MovieSceneTracks/Private/TrackInstances/MovieSceneCameraShakeTrackInstance.cpp b/Engine/Source/Runtime/MovieSceneTracks/Private/TrackInstances/MovieSceneCameraShakeTrackInstance.cpp
--- a/Engine/Source/Runtime/MovieSceneTracks/Private/TrackInstances/MovieSceneCameraShakeTrackInstance.cpp
+++ b/Engine/Source/Runtime/MovieSceneTracks/Private/TrackInstances/MovieSceneCameraShakeTrackInstance.cpp
@@ -92,7 +92,7 @@ void FMovieSceneCameraShakeTrackInstance::Update(EMovieSceneUpdateData& UpdateDa
 						{
 							// make it root so GC doesn't take it away
 							ShakeInst->AddToRoot();
-							ShakeInst->SetTempCameraAnimActor(GetTempCameraActor());
+							ShakeInst->SetTempCameraAnimActor(GetTempCameraActor(Player));
 							ShakeInst->PlayShake(nullptr, ActiveSection->PlayScale, ActiveSection->PlaySpace, ActiveSection->UserDefinedPlaySpace);
 							if (ShakeInst->AnimInst)
 							{
@@ -113,7 +113,7 @@ void FMovieSceneCameraShakeTrackInstance::Update(EMovieSceneUpdateData& UpdateDa
 
 						// prepare temp camera actor by resetting it
 						{
-							ACameraActor* const CameraActor = GetTempCameraActor();
+							ACameraActor* const CameraActor = GetTempCameraActor(Player);
 							ACameraActor const* const DefaultCamActor = GetDefault<ACameraActor>();
 							if (CameraActor && DefaultCamActor)
 							{
@@ -147,7 +147,7 @@ void FMovieSceneCameraShakeTrackInstance::Update(EMovieSceneUpdateData& UpdateDa
 
 						// harvest PP changes
 						{
-							UCameraComponent* AnimCamComp = GetTempCameraActor()->GetCameraComponent();
+							UCameraComponent* AnimCamComp = GetTempCameraActor(Player)->GetCameraComponent();
 							if (AnimCamComp)
 							{
 								SectionInstanceData.PostProcessingBlendWeight = AnimCamComp->PostProcessBlendWeight;
@@ -183,7 +183,7 @@ void FMovieSceneCameraShakeTrackInstance::Update(EMovieSceneUpdateData& UpdateDa
 }
 
 
-ACameraActor* FMovieSceneCameraShakeTrackInstance::GetTempCameraActor()
+ACameraActor* FMovieSceneCameraShakeTrackInstance::GetTempCameraActor(IMovieScenePlayer& InMovieScenePlayer)
 {
 	if (TempCameraActor.IsValid() == false)
 	{
@@ -191,7 +191,7 @@ ACameraActor* FMovieSceneCameraShakeTrackInstance::GetTempCameraActor()
 		FActorSpawnParameters SpawnInfo;
 		SpawnInfo.SpawnCollisionHandlingOverride = ESpawnActorCollisionHandlingMethod::AlwaysSpawn;
 		SpawnInfo.ObjectFlags |= RF_Transient;	// We never want to save these temp actors into a map
-		ACameraActor* const Cam = GWorld->SpawnActor<ACameraActor>(SpawnInfo);		// #fixme, GWorld!
+		ACameraActor* const Cam = InMovieScenePlayer.GetPlaybackContext()->GetWorld()->SpawnActor<ACameraActor>(SpawnInfo);
 		if (Cam)
 		{
 #if WITH_EDITOR
diff --git a/Engine/Source/Runtime/MovieSceneTracks/Private/TrackInstances/MovieSceneParticleTrackInstance.cpp b/Engine/Source/Runtime/MovieSceneTracks/Private/TrackInstances/MovieSceneParticleTrackInstance.cpp
--- a/Engine/Source/Runtime/MovieSceneTracks/Private/TrackInstances/MovieSceneParticleTrackInstance.cpp
+++ b/Engine/Source/Runtime/MovieSceneTracks/Private/TrackInstances/MovieSceneParticleTrackInstance.cpp
@@ -29,8 +29,7 @@ void FMovieSceneParticleTrackInstance::Update(EMovieSceneUpdateData& UpdateData,
 	// @todo Sequencer We need something analagous to Matinee 1's particle replay tracks
 	// What we have here is simple toggling/triggering
 
-	float LastPosition = UpdateData.bLooped ? UpdateData.Position : UpdateData.LastPosition;
-	if (UpdateData.Position >= LastPosition && Player.GetPlaybackStatus() == EMovieScenePlayerStatus::Playing)
+	if (UpdateData.Position >= UpdateData.LastPosition && Player.GetPlaybackStatus() == EMovieScenePlayerStatus::Playing)
 	{
 		const TArray<UMovieSceneSection*> Sections = ParticleTrack->GetAllParticleSections();
 		EParticleKey::Type ParticleKey = EParticleKey::Deactivate;
@@ -46,7 +45,7 @@ void FMovieSceneParticleTrackInstance::Update(EMovieSceneUpdateData& UpdateData,
 				if ( ParticleKeyCurve.IsKeyHandleValid( PreviousHandle ) )
 				{
 					FIntegralKey& PreviousKey = ParticleKeyCurve.GetKey( PreviousHandle );
-					if ( PreviousKey.Time >= LastPosition )
+					if ( PreviousKey.Time >= UpdateData.LastPosition )
 					{
 						ParticleKey = (EParticleKey::Type)PreviousKey.Value;
 						bKeyFound = true;
diff --git a/Engine/Source/Runtime/MovieSceneTracks/Private/TrackInstances/MovieSceneSlomoTrackInstance.cpp b/Engine/Source/Runtime/MovieSceneTracks/Private/TrackInstances/MovieSceneSlomoTrackInstance.cpp
--- a/Engine/Source/Runtime/MovieSceneTracks/Private/TrackInstances/MovieSceneSlomoTrackInstance.cpp
+++ b/Engine/Source/Runtime/MovieSceneTracks/Private/TrackInstances/MovieSceneSlomoTrackInstance.cpp
@@ -16,7 +16,7 @@ FMovieSceneSlomoTrackInstance::FMovieSceneSlomoTrackInstance(UMovieSceneSlomoTra
 	
 void FMovieSceneSlomoTrackInstance::RestoreState(const TArray<TWeakObjectPtr<UObject>>& RuntimeObjects, IMovieScenePlayer& Player, FMovieSceneSequenceInstance& SequenceInstance)
 {
-	AWorldSettings* WorldSettings = GWorld->GetWorldSettings();
+	AWorldSettings* WorldSettings = Player.GetPlaybackContext()->GetWorld()->GetWorldSettings();
 
 	if (WorldSettings == nullptr)
 	{
@@ -28,7 +28,7 @@ void FMovieSceneSlomoTrackInstance::RestoreState(const TArray<TWeakObjectPtr<UOb
 
 void FMovieSceneSlomoTrackInstance::SaveState(const TArray<TWeakObjectPtr<UObject>>& RuntimeObjects, IMovieScenePlayer& Player, FMovieSceneSequenceInstance& SequenceInstance)
 {
-	AWorldSettings* WorldSettings = GWorld->GetWorldSettings();
+	AWorldSettings* WorldSettings = Player.GetPlaybackContext()->GetWorld()->GetWorldSettings();
 
 	if (WorldSettings == nullptr)
 	{
@@ -43,12 +43,12 @@ void FMovieSceneSlomoTrackInstance::SaveState(const TArray<TWeakObjectPtr<UObjec
 
 void FMovieSceneSlomoTrackInstance::Update(EMovieSceneUpdateData& UpdateData, const TArray<TWeakObjectPtr<UObject>>& RuntimeObjects, IMovieScenePlayer& Player, FMovieSceneSequenceInstance& SequenceInstance)
 {
-	if (!ShouldBeApplied())
+	if (!ShouldBeApplied(Player))
 	{
 		return;
 	}
 
-	AWorldSettings* WorldSettings = GWorld->GetWorldSettings();
+	AWorldSettings* WorldSettings = Player.GetPlaybackContext()->GetWorld()->GetWorldSettings();
 
 	if (WorldSettings == nullptr)
 	{
@@ -70,14 +70,14 @@ void FMovieSceneSlomoTrackInstance::Update(EMovieSceneUpdateData& UpdateData, co
 /* IMovieSceneTrackInstance implementation
  *****************************************************************************/
 
-bool FMovieSceneSlomoTrackInstance::ShouldBeApplied() const
+bool FMovieSceneSlomoTrackInstance::ShouldBeApplied(IMovieScenePlayer& InMovieScenePlayer) const
 {
 	if (GIsEditor)
 	{
 		return true;
 	}
 
-	if (GWorld->GetNetMode() == NM_Client)
+	if (InMovieScenePlayer.GetPlaybackContext()->GetWorld()->GetNetMode() == NM_Client)
 	{
 		return false;
 	}
diff --git a/Engine/Source/Runtime/MovieSceneTracks/Private/TrackInstances/MovieSceneSubTrackInstance.cpp b/Engine/Source/Runtime/MovieSceneTracks/Private/TrackInstances/MovieSceneSubTrackInstance.cpp
--- a/Engine/Source/Runtime/MovieSceneTracks/Private/TrackInstances/MovieSceneSubTrackInstance.cpp
+++ b/Engine/Source/Runtime/MovieSceneTracks/Private/TrackInstances/MovieSceneSubTrackInstance.cpp
@@ -5,6 +5,7 @@
 #include "MovieSceneSubTrackInstance.h"
 #include "MovieSceneSequenceInstance.h"
 #include "MovieSceneSubSection.h"
+#include "MovieSceneCinematicShotTrack.h"
 #include "MovieSceneSequence.h"
 #include "IMovieScenePlayer.h"
 
@@ -251,6 +252,7 @@ void FMovieSceneSubTrackInstance::UpdateSection( EMovieSceneUpdateData& UpdateDa
 	SubUpdateData.UpdatePass = UpdateData.UpdatePass;
 	SubUpdateData.bPreroll = UpdateData.Position < SubSection->GetStartTime();
 	SubUpdateData.bSubSceneDeactivate = bSectionWasDeactivated;
+	SubUpdateData.bUpdateCameras = SubTrack->GetClass() == UMovieSceneCinematicShotTrack::StaticClass();
 
 	// update sub sections
 
diff --git a/Engine/Source/Runtime/MovieSceneTracks/Private/Tracks/MovieSceneEventTrack.cpp b/Engine/Source/Runtime/MovieSceneTracks/Private/Tracks/MovieSceneEventTrack.cpp
--- a/Engine/Source/Runtime/MovieSceneTracks/Private/Tracks/MovieSceneEventTrack.cpp
+++ b/Engine/Source/Runtime/MovieSceneTracks/Private/Tracks/MovieSceneEventTrack.cpp
@@ -44,6 +44,14 @@ void UMovieSceneEventTrack::TriggerEvents(float Position, float LastPosition, IM
 		return;
 	}
 
+	// Don't allow events to fire when playback is in a stopped state. This can occur when stopping 
+	// playback and returning the current position to the start of playback. It's not desireable to have 
+	// all the events from the last playback position to the start of playback be fired.
+	if (Player.GetPlaybackStatus() == EMovieScenePlayerStatus::Stopped)
+	{
+		return;
+	}
+
 	bool Backwards = Position < LastPosition;
 
 	if ((!Backwards && !bFireEventsWhenForwards) ||
diff --git a/Engine/Source/Runtime/MovieSceneTracks/Public/Sections/MovieScene3DTransformSection.h b/Engine/Source/Runtime/MovieSceneTracks/Public/Sections/MovieScene3DTransformSection.h
--- a/Engine/Source/Runtime/MovieSceneTracks/Public/Sections/MovieScene3DTransformSection.h
+++ b/Engine/Source/Runtime/MovieSceneTracks/Public/Sections/MovieScene3DTransformSection.h
@@ -19,6 +19,17 @@ namespace EKey3DTransformChannel
 	};
 }
 
+#if WITH_EDITORONLY_DATA
+/** Visibility options for 3d trajectory. */
+UENUM()
+enum class EShow3DTrajectory
+{
+	EST_OnlyWhenSelected UMETA(DisplayName="Only When Selected"),
+	EST_Always UMETA(DisplayName="Always"),
+	EST_Never UMETA(DisplayName="Never"),
+};
+#endif
+
 /**
 * Stores information about a transform for the purpose of adding keys to a transform section
 */
@@ -216,6 +227,13 @@ public:
 	 */
 	MOVIESCENETRACKS_API FRichCurve& GetScaleCurve( EAxis::Type Axis );
 
+	/**
+	 * Return the trajectory visibility
+	 */
+#if WITH_EDITORONLY_DATA
+	MOVIESCENETRACKS_API EShow3DTrajectory GetShow3DTrajectory() { return Show3DTrajectory; }
+#endif
+
 public:
 
 	// UMovieSceneSection interface
@@ -249,4 +267,10 @@ private:
 	/** Scale curves */
 	UPROPERTY()
 	FRichCurve Scale[3];
+
+#if WITH_EDITORONLY_DATA
+	/** Whether to show the 3d trajectory */
+	UPROPERTY(EditAnywhere, DisplayName = "Show 3D Trajectory", Category = "Transform")
+	TEnumAsByte<EShow3DTrajectory> Show3DTrajectory;
+#endif
 };
\ No newline at end of file
diff --git a/Engine/Source/Runtime/MovieSceneTracks/Public/Sections/MovieSceneSubSection.h b/Engine/Source/Runtime/MovieSceneTracks/Public/Sections/MovieSceneSubSection.h
--- a/Engine/Source/Runtime/MovieSceneTracks/Public/Sections/MovieSceneSubSection.h
+++ b/Engine/Source/Runtime/MovieSceneTracks/Public/Sections/MovieSceneSubSection.h
@@ -46,7 +46,7 @@ public:
 	static UMovieSceneSubSection* GetRecordingSection();
 
 	/** Get the actor we are targeting for recording */
-	static const FString& GetActorToRecord();
+	static AActor* GetActorToRecord();
 
 	/** Check if we are primed for recording */
 	static bool IsSetAsRecording();
@@ -82,10 +82,10 @@ public:
 		TargetPathToRecordTo.Path = Path;
 	}
 
-	/** Set the target actor name to record */
-	void SetNameOfActorToRecord(const FString& ActorName)
+	/** Set the target actor to record */
+	void SetActorToRecord(AActor* InActorToRecord)
 	{
-		NameOfActorToRecord = ActorName;
+		ActorToRecord = InActorToRecord;
 	}
 
 	virtual UMovieSceneSection* SplitSection( float SplitTime ) override;
@@ -117,9 +117,9 @@ protected:
 	UPROPERTY(EditAnywhere, Category="Sequence")
 	UMovieSceneSequence* SubSequence;
 
-	/** Target actor name to record */
+	/** Target actor to record */
 	UPROPERTY(EditAnywhere, Category="Sequence Recording")
-	FString NameOfActorToRecord;
+	TLazyObjectPtr<AActor> ActorToRecord;
 
 	/** Target name of sequence to try to record to (will record automatically to another if this already exists) */
 	UPROPERTY(EditAnywhere, Category="Sequence Recording")
diff --git a/Engine/Source/Runtime/MovieSceneTracks/Public/TrackInstances/MovieSceneCameraAnimTrackInstance.h b/Engine/Source/Runtime/MovieSceneTracks/Public/TrackInstances/MovieSceneCameraAnimTrackInstance.h
--- a/Engine/Source/Runtime/MovieSceneTracks/Public/TrackInstances/MovieSceneCameraAnimTrackInstance.h
+++ b/Engine/Source/Runtime/MovieSceneTracks/Public/TrackInstances/MovieSceneCameraAnimTrackInstance.h
@@ -45,7 +45,7 @@ private:
 
 	TWeakObjectPtr<ACameraActor> TempCameraActor;
 
-	ACameraActor* GetTempCameraActor();
+	ACameraActor* GetTempCameraActor(IMovieScenePlayer& InMovieScenePlayer);
 
 	TMap<UMovieSceneCameraAnimSection*, FMovieSceneCameraAnimSectionInstanceData> SectionInstanceDataMap;
 };
diff --git a/Engine/Source/Runtime/MovieSceneTracks/Public/TrackInstances/MovieSceneCameraShakeTrackInstance.h b/Engine/Source/Runtime/MovieSceneTracks/Public/TrackInstances/MovieSceneCameraShakeTrackInstance.h
--- a/Engine/Source/Runtime/MovieSceneTracks/Public/TrackInstances/MovieSceneCameraShakeTrackInstance.h
+++ b/Engine/Source/Runtime/MovieSceneTracks/Public/TrackInstances/MovieSceneCameraShakeTrackInstance.h
@@ -47,5 +47,5 @@ private:
 
 	TMap<UMovieSceneCameraShakeSection*, FMovieSceneCameraShakeSectionInstanceData> SectionInstanceDataMap;
 
-	ACameraActor* GetTempCameraActor();
+	ACameraActor* GetTempCameraActor(IMovieScenePlayer& InMovieScenePlayer);
 };
diff --git a/Engine/Source/Runtime/MovieSceneTracks/Public/TrackInstances/MovieSceneSlomoTrackInstance.h b/Engine/Source/Runtime/MovieSceneTracks/Public/TrackInstances/MovieSceneSlomoTrackInstance.h
--- a/Engine/Source/Runtime/MovieSceneTracks/Public/TrackInstances/MovieSceneSlomoTrackInstance.h
+++ b/Engine/Source/Runtime/MovieSceneTracks/Public/TrackInstances/MovieSceneSlomoTrackInstance.h
@@ -42,7 +42,7 @@ protected:
 	AWorldSettings* GetWorld() const;
 
 	/** Check whether the slow motion should be applied. */
-	bool ShouldBeApplied() const;
+	bool ShouldBeApplied(IMovieScenePlayer& InMovieScenePlayer) const;
 
 private:
 
diff --git a/Engine/Source/Runtime/Slate/Private/Widgets/Input/SEditableLabel.cpp b/Engine/Source/Runtime/Slate/Private/Widgets/Input/SEditableLabel.cpp
--- a/Engine/Source/Runtime/Slate/Private/Widgets/Input/SEditableLabel.cpp
+++ b/Engine/Source/Runtime/Slate/Private/Widgets/Input/SEditableLabel.cpp
@@ -114,7 +114,7 @@ FReply SEditableLabel::OnKeyDown(const FGeometry& MyGeometry, const FKeyEvent& I
 
 bool SEditableLabel::SupportsKeyboardFocus() const
 {
-	return true;
+	return CanEditAttribute.Get();
 }
 
 
diff --git a/Engine/Source/Runtime/UMG/Private/Animation/UMGSequencePlayer.cpp b/Engine/Source/Runtime/UMG/Private/Animation/UMGSequencePlayer.cpp
--- a/Engine/Source/Runtime/UMG/Private/Animation/UMGSequencePlayer.cpp
+++ b/Engine/Source/Runtime/UMG/Private/Animation/UMGSequencePlayer.cpp
@@ -146,7 +146,7 @@ void UUMGSequencePlayer::PlayInternal(double StartAtTime, double EndAtTime, doub
 
 	if (PlayMode == EUMGSequencePlayMode::Reverse)
 	{
-		// When playing in reverse count substract the start time from the end.
+		// When playing in reverse count subtract the start time from the end.
 		TimeCursorPosition = CurrentPlayRange.GetUpperBoundValue() - StartAtTime;
 	}
 	else
@@ -178,7 +178,7 @@ void UUMGSequencePlayer::PlayInternal(double StartAtTime, double EndAtTime, doub
 void UUMGSequencePlayer::Play(float StartAtTime, int32 InNumLoopsToPlay, EUMGSequencePlayMode::Type InPlayMode)
 {
 	double SubAnimStartTime = 0.0;
-	double SubAnimEndTime = TimeRange.GetUpperBoundValue();
+	double SubAnimEndTime = TimeRange.Size<float>();
 
 	PlayInternal(StartAtTime, 0.0, SubAnimStartTime, SubAnimEndTime, InNumLoopsToPlay, InPlayMode);
 }
@@ -186,7 +186,7 @@ void UUMGSequencePlayer::Play(float StartAtTime, int32 InNumLoopsToPlay, EUMGSeq
 void UUMGSequencePlayer::PlayTo(float StartAtTime, float EndAtTime, int32 InNumLoopsToPlay, EUMGSequencePlayMode::Type InPlayMode)
 {
 	double SubAnimStartTime = 0.0;
-	double SubAnimEndTime = TimeRange.GetUpperBoundValue();
+	double SubAnimEndTime = TimeRange.Size<float>();
 
 	PlayInternal(StartAtTime, EndAtTime, SubAnimStartTime, SubAnimEndTime, InNumLoopsToPlay, InPlayMode);
 }
diff --git a/Engine/Source/Runtime/UMG/Public/Animation/UMGSequencePlayer.h b/Engine/Source/Runtime/UMG/Public/Animation/UMGSequencePlayer.h
--- a/Engine/Source/Runtime/UMG/Public/Animation/UMGSequencePlayer.h
+++ b/Engine/Source/Runtime/UMG/Public/Animation/UMGSequencePlayer.h
@@ -45,7 +45,7 @@ public:
 
 	/** IMovieScenePlayer interface */
 	virtual void GetRuntimeObjects( TSharedRef<FMovieSceneSequenceInstance> MovieSceneInstance, const FGuid& ObjectHandle, TArray<TWeakObjectPtr<UObject>>& OutObjects ) const override;
-	virtual void UpdateCameraCut(UObject* CameraObject, UObject* UnlockIfCameraObject, bool bJumpCut) const override {}
+	virtual void UpdateCameraCut(UObject* CameraObject, UObject* UnlockIfCameraObject, bool bJumpCut) override {}
 	virtual void SetViewportSettings(const TMap<FViewportClient*, EMovieSceneViewportParams>& ViewportParamsMap) override {}
 	virtual void GetViewportSettings(TMap<FViewportClient*, EMovieSceneViewportParams>& ViewportParamsMap) const override {}
 	virtual void AddOrUpdateMovieSceneInstance( class UMovieSceneSection& MovieSceneSection, TSharedRef<FMovieSceneSequenceInstance> InstanceToAdd ) override {}
