Movie Scene Capture: Fix non unity build.

diff --git a/Engine/Source/Runtime/MovieSceneCapture/Private/MovieSceneCaptureModule.cpp b/Engine/Source/Runtime/MovieSceneCapture/Private/MovieSceneCaptureModule.cpp
--- a/Engine/Source/Runtime/MovieSceneCapture/Private/MovieSceneCaptureModule.cpp
+++ b/Engine/Source/Runtime/MovieSceneCapture/Private/MovieSceneCaptureModule.cpp
@@ -5,6 +5,7 @@
 #include "MovieSceneCaptureModule.h"
 #include "JsonObjectConverter.h"
 #include "ActiveMovieSceneCaptures.h"
+#include "SceneViewport.h"
 
 #include "Protocols/ImageSequenceProtocol.h"
 #include "Protocols/CompositionGraphCaptureProtocol.h"
