Renamed UAutomatedLevelSequenceCapture's StartFrame and EndFrame to CustomStartFrame/CustomEndFrame to reflect that it is user-overriden. This will also resolve upgrade issues due to the variable datatype changing.

diff --git a/Engine/Source/Editor/Sequencer/Private/Sequencer.cpp b/Engine/Source/Editor/Sequencer/Private/Sequencer.cpp
--- a/Engine/Source/Editor/Sequencer/Private/Sequencer.cpp
+++ b/Engine/Source/Editor/Sequencer/Private/Sequencer.cpp
@@ -2173,12 +2173,12 @@ void FSequencer::RenderMovieInternal(TRange<FFrameNumber> Range, bool bSetFrameO
 	{
 		if (!MovieSceneCapture->bUseCustomStartFrame)
 		{
-			MovieSceneCapture->StartFrame = RoundedStartFrame;
+			MovieSceneCapture->CustomStartFrame = RoundedStartFrame;
 		}
 
 		if (!MovieSceneCapture->bUseCustomEndFrame)
 		{
-			MovieSceneCapture->EndFrame = RoundedEndFrame;
+			MovieSceneCapture->CustomEndFrame = RoundedEndFrame;
 		}
 	}
 
diff --git a/Engine/Source/Runtime/MovieSceneCapture/Private/AutomatedLevelSequenceCapture.cpp b/Engine/Source/Runtime/MovieSceneCapture/Private/AutomatedLevelSequenceCapture.cpp
--- a/Engine/Source/Runtime/MovieSceneCapture/Private/AutomatedLevelSequenceCapture.cpp
+++ b/Engine/Source/Runtime/MovieSceneCapture/Private/AutomatedLevelSequenceCapture.cpp
@@ -77,9 +77,9 @@ UAutomatedLevelSequenceCapture::UAutomatedLevelSequenceCapture(const FObjectInit
 	}
 #else
 	bUseCustomStartFrame = false;
-	StartFrame = 0;
+	CustomStartFrame = 0;
 	bUseCustomEndFrame = false;
-	EndFrame = 1;
+	CustomEndFrame = 1;
 	WarmUpFrameCount = 0;
 	DelayBeforeWarmUp = 0.0f;
 	DelayBeforeShotWarmUp = 0.0f;
@@ -131,14 +131,14 @@ void UAutomatedLevelSequenceCapture::Initialize(TSharedPtr<FSceneViewport> InVie
 		if( FParse::Value( FCommandLine::Get(), TEXT( "-MovieStartFrame=" ), StartFrameOverride ) )
 		{
 			bUseCustomStartFrame = true;
-			StartFrame = StartFrameOverride;
+			CustomStartFrame = StartFrameOverride;
 		}
 
 		int32 EndFrameOverride;
 		if( FParse::Value( FCommandLine::Get(), TEXT( "-MovieEndFrame=" ), EndFrameOverride ) )
 		{
 			bUseCustomEndFrame = true;
-			EndFrame = EndFrameOverride;
+			CustomEndFrame = EndFrameOverride;
 		}
 
 		int32 WarmUpFrameCountOverride;
@@ -435,7 +435,7 @@ void UAutomatedLevelSequenceCapture::SetupFrameRange()
 
 				if( bUseCustomStartFrame )
 				{
-					PlaybackStartFrame = Settings.bUseRelativeFrameNumbers ? ( PlaybackStartFrame + StartFrame ) : StartFrame;
+					PlaybackStartFrame = Settings.bUseRelativeFrameNumbers ? ( PlaybackStartFrame + CustomStartFrame ) : CustomStartFrame;
 				}
 
 				if( !Settings.bUseRelativeFrameNumbers )
@@ -448,7 +448,7 @@ void UAutomatedLevelSequenceCapture::SetupFrameRange()
 
 				if( bUseCustomEndFrame )
 				{
-				 	PlaybackEndFrame = FMath::Max( PlaybackStartFrame, Settings.bUseRelativeFrameNumbers ? ( PlaybackEndFrame + EndFrame ) : EndFrame );
+				 	PlaybackEndFrame = FMath::Max( PlaybackStartFrame, Settings.bUseRelativeFrameNumbers ? ( PlaybackEndFrame + CustomEndFrame ) : CustomEndFrame );
 				}
 
 				RemainingWarmUpFrames = FMath::Max( WarmUpFrameCount, 0 );
@@ -664,8 +664,8 @@ void UAutomatedLevelSequenceCapture::LoadFromConfig()
 
 void UAutomatedLevelSequenceCapture::SaveToConfig()
 {
-	FFrameNumber CurrentStartFrame = StartFrame;
-	FFrameNumber CurrentEndFrame = EndFrame;
+	FFrameNumber CurrentStartFrame = CustomStartFrame;
+	FFrameNumber CurrentEndFrame = CustomEndFrame;
 	bool bRestoreFrameOverrides = RestoreFrameOverrides();
 
 	BurnInOptions->SaveConfig();
@@ -694,13 +694,13 @@ bool UAutomatedLevelSequenceCapture::RestoreFrameOverrides()
 	bool bAnySet = CachedStartFrame.IsSet() || CachedEndFrame.IsSet() || bCachedUseCustomStartFrame.IsSet() || bCachedUseCustomEndFrame.IsSet();
 	if (CachedStartFrame.IsSet())
 	{
-		StartFrame = CachedStartFrame.GetValue();
+		CustomStartFrame = CachedStartFrame.GetValue();
 		CachedStartFrame.Reset();
 	}
 
 	if (CachedEndFrame.IsSet())
 	{
-		EndFrame = CachedEndFrame.GetValue();
+		CustomEndFrame = CachedEndFrame.GetValue();
 		CachedEndFrame.Reset();
 	}
 
@@ -721,13 +721,13 @@ bool UAutomatedLevelSequenceCapture::RestoreFrameOverrides()
 
 void UAutomatedLevelSequenceCapture::SetFrameOverrides(FFrameNumber InStartFrame, FFrameNumber InEndFrame)
 {
-	CachedStartFrame = StartFrame;
-	CachedEndFrame = EndFrame;
+	CachedStartFrame = CustomStartFrame;
+	CachedEndFrame = CustomEndFrame;
 	bCachedUseCustomStartFrame = bUseCustomStartFrame;
 	bCachedUseCustomEndFrame = bUseCustomEndFrame;
 
-	StartFrame = InStartFrame;
-	EndFrame = InEndFrame;
+	CustomStartFrame = InStartFrame;
+	CustomEndFrame = InEndFrame;
 	bUseCustomStartFrame = true;
 	bUseCustomEndFrame = true;
 }
diff --git a/Engine/Source/Runtime/MovieSceneCapture/Public/AutomatedLevelSequenceCapture.h b/Engine/Source/Runtime/MovieSceneCapture/Public/AutomatedLevelSequenceCapture.h
--- a/Engine/Source/Runtime/MovieSceneCapture/Public/AutomatedLevelSequenceCapture.h
+++ b/Engine/Source/Runtime/MovieSceneCapture/Public/AutomatedLevelSequenceCapture.h
@@ -31,16 +31,16 @@ public:
 	bool bUseCustomStartFrame;
 
 	/** Frame number to start capturing.  The frame number range depends on whether the bUseRelativeFrameNumbers option is enabled. */
-	UPROPERTY(config, EditAnywhere, Category=Animation, AdvancedDisplay, meta=(EditCondition="bUseCustomStartFrame"))
-	FFrameNumber StartFrame;
+	UPROPERTY(config, EditAnywhere, Category=Animation, AdvancedDisplay, meta=(EditCondition="bUseCustomStartFrame", DisplayName="Start Frame"))
+	FFrameNumber CustomStartFrame;
 
 	/** When enabled, the EndFrame setting will override the default ending frame number */
 	UPROPERTY(config, EditAnywhere, Category=Animation, AdvancedDisplay)
 	bool bUseCustomEndFrame;
 
 	/** Frame number to end capturing.  The frame number range depends on whether the bUseRelativeFrameNumbers option is enabled. */
-	UPROPERTY(config, EditAnywhere, Category=Animation, AdvancedDisplay, meta=(EditCondition="bUseCustomEndFrame"))
-	FFrameNumber EndFrame;
+	UPROPERTY(config, EditAnywhere, Category=Animation, AdvancedDisplay, meta=(EditCondition="bUseCustomEndFrame", DisplayName="End Frame"))
+	FFrameNumber CustomEndFrame;
 
 	/** The number of extra frames to play before the sequence's start frame, to "warm up" the animation.  This is useful if your
 	    animation contains particles or other runtime effects that are spawned into the scene earlier than your capture start frame */
