Movie Scene Capture: Fix crash switching capture protocols.

diff --git a/Engine/Source/Runtime/MovieSceneCapture/Public/MovieSceneCaptureProtocolRegistry.h b/Engine/Source/Runtime/MovieSceneCapture/Public/MovieSceneCaptureProtocolRegistry.h
--- a/Engine/Source/Runtime/MovieSceneCapture/Public/MovieSceneCaptureProtocolRegistry.h
+++ b/Engine/Source/Runtime/MovieSceneCapture/Public/MovieSceneCaptureProtocolRegistry.h
@@ -63,7 +63,7 @@ public:
 		{
 			if (Info->SettingsClassType)
 			{
-				return NewObject<UMovieSceneCaptureProtocolSettings>(Outer, Info->SettingsClassType, FName(TEXT("MovieSceneCaptureProtocol")));
+				return NewObject<UMovieSceneCaptureProtocolSettings>(Outer, Info->SettingsClassType);
 			}
 		}
 		return nullptr;
