Movie Scene Capture: Removed inadvertent config option on InheritedCommandLineArguments so that they don't get saved and override the actual command line args.

diff --git a/Engine/Source/Runtime/MovieSceneCapture/Public/MovieSceneCapture.h b/Engine/Source/Runtime/MovieSceneCapture/Public/MovieSceneCapture.h
--- a/Engine/Source/Runtime/MovieSceneCapture/Public/MovieSceneCapture.h
+++ b/Engine/Source/Runtime/MovieSceneCapture/Public/MovieSceneCapture.h
@@ -95,7 +95,7 @@ public:
 	FString AdditionalCommandLineArguments;
 
 	/** Command line arguments inherited from this process */
-	UPROPERTY(config, EditAnywhere, BlueprintReadWrite, transient, Category=General, AdvancedDisplay, meta=(EditCondition=bUseSeparateProcess))
+	UPROPERTY(EditAnywhere, BlueprintReadWrite, transient, Category=General, AdvancedDisplay, meta=(EditCondition=bUseSeparateProcess))
 	FString InheritedCommandLineArguments;
 
 	/** Event that is fired after we've finished capturing */
