MovieSceneCapture: Normalize filename for the MovieFolder path and the Output path. Output an error and disable the UI from allowing movie scene capture if the output directory is read only.

diff --git a/Engine/Source/Editor/MovieSceneCaptureDialog/Private/MovieSceneCaptureDialogModule.cpp b/Engine/Source/Editor/MovieSceneCaptureDialog/Private/MovieSceneCaptureDialogModule.cpp
--- a/Engine/Source/Editor/MovieSceneCaptureDialog/Private/MovieSceneCaptureDialogModule.cpp
+++ b/Engine/Source/Editor/MovieSceneCaptureDialog/Private/MovieSceneCaptureDialogModule.cpp
@@ -23,6 +23,8 @@
 
 #include "HAL/IConsoleManager.h"
 #include "HAL/PlatformProcess.h"
+#include "HAL/PlatformFilemanager.h"
+#include "HAL/FileManager.h"
 #include "Misc/CommandLine.h"
 #include "Misc/FileHelper.h"
 #include "Misc/App.h"
@@ -875,6 +877,21 @@ class FMovieSceneCaptureDialogModule : public IMovieSceneCaptureDialogModule
 			return LOCTEXT("AlreadyCapturing", "There is already a movie scene capture process open. Please close it and try again.");
 		}
 
+		FString OutputDirectory = CaptureObject->Settings.OutputDirectory.Path;
+		FPaths::NormalizeFilename(OutputDirectory);
+
+		if (!IFileManager::Get().DirectoryExists(*OutputDirectory))
+		{
+			if (!IFileManager::Get().MakeDirectory(*OutputDirectory))
+			{
+				return FText::Format(LOCTEXT( "InvalidDirectory", "Invalid output directory: {0}"), FText::FromString(OutputDirectory) );
+			}
+		}
+		else if (IFileManager::Get().IsReadOnly(*OutputDirectory))
+		{
+			return FText::Format(LOCTEXT( "ReadOnlyDirectory", "Read only output directory: {0}"), FText::FromString(OutputDirectory) );
+		}
+
 		// Prompt the user to save their changes so that they'll be in the movie, since we're not saving temporary copies of the level.
 		bool bPromptUserToSave = true;
 		bool bSaveMapPackages = true;
diff --git a/Engine/Source/Runtime/MovieSceneCapture/Private/MovieSceneCapture.cpp b/Engine/Source/Runtime/MovieSceneCapture/Private/MovieSceneCapture.cpp
--- a/Engine/Source/Runtime/MovieSceneCapture/Private/MovieSceneCapture.cpp
+++ b/Engine/Source/Runtime/MovieSceneCapture/Private/MovieSceneCapture.cpp
@@ -109,6 +109,19 @@ void UMovieSceneCapture::Initialize(TSharedPtr<FSceneViewport> InSceneViewport,
 		if( FParse::Value( FCommandLine::Get(), TEXT( "-MovieFolder=" ), OutputPathOverride ) )
 		{
 			Settings.OutputDirectory.Path = OutputPathOverride;
+			FPaths::NormalizeFilename(Settings.OutputDirectory.Path);
+
+			if (!IFileManager::Get().DirectoryExists(*Settings.OutputDirectory.Path))
+			{
+				if (!IFileManager::Get().MakeDirectory(*Settings.OutputDirectory.Path))
+				{
+					UE_LOG(LogMovieSceneCapture, Error, TEXT("Invalid output directory: %s."), *Settings.OutputDirectory.Path);
+				}
+			}
+			else if (IFileManager::Get().IsReadOnly(*Settings.OutputDirectory.Path))
+			{
+				UE_LOG(LogMovieSceneCapture, Error, TEXT("Read only output directory: %s."), *Settings.OutputDirectory.Path);
+			}
 		}
 
 		FString OutputNameOverride;
@@ -331,7 +344,10 @@ void UMovieSceneCapture::EnsureFileWritable(const FString& File) const
 
 FString UMovieSceneCapture::GenerateFilename(const FFrameMetrics& FrameMetrics, const TCHAR* Extension) const
 {
-	const FString BaseFilename = ResolveFileFormat(Settings.OutputDirectory.Path, FrameMetrics) / ResolveFileFormat(Settings.OutputFormat, FrameMetrics);
+	FString OutputDirectoryPath = Settings.OutputDirectory.Path;
+	FPaths::NormalizeFilename(OutputDirectoryPath);
+
+	const FString BaseFilename = ResolveFileFormat(OutputDirectoryPath, FrameMetrics) / ResolveFileFormat(Settings.OutputFormat, FrameMetrics);
 
 	FString ThisTry = BaseFilename + Extension;
 
