Movie Scene Capture: Move linux warning to Initialize to avoid spamming log every frame.

diff --git a/Engine/Source/Runtime/MovieSceneCapture/Private/VideoCaptureProtocol.cpp b/Engine/Source/Runtime/MovieSceneCapture/Private/VideoCaptureProtocol.cpp
--- a/Engine/Source/Runtime/MovieSceneCapture/Private/VideoCaptureProtocol.cpp
+++ b/Engine/Source/Runtime/MovieSceneCapture/Private/VideoCaptureProtocol.cpp
@@ -14,6 +14,10 @@ bool FVideoCaptureProtocol::Initialize(const FCaptureProtocolInitSettings& InSet
 		return false;
 	}
 
+#if PLATFORM_UNIX
+UE_LOG(LogInit, Warning, TEXT("Writing movies is not currently supported on Linux"));
+#endif
+
 	ConditionallyCreateWriter(Host);
 
 	return AVIWriters.Num() && AVIWriters.Last()->IsCapturing();
@@ -25,7 +29,6 @@ void FVideoCaptureProtocol::ConditionallyCreateWriter(const ICaptureProtocolHost
 	static const TCHAR* Extension = TEXT(".mov");
 #elif PLATFORM_UNIX
 	static const TCHAR* Extension = TEXT(".unsupp");
-	UE_LOG(LogInit, Warning, TEXT("Writing movies is not currently supported on Linux"));
 	return;
 #else
 	static const TCHAR* Extension = TEXT(".avi");
