CompositionGraphCaptureProtocol now doesn't stop on "," separators so that you can specify multiple passes at once. This fixes a regression and doesn't add any new behavior to CompositionGraphCaptureProtocol.

diff --git a/Engine/Source/Runtime/MovieSceneCapture/Private/CompositionGraphCaptureProtocol.cpp b/Engine/Source/Runtime/MovieSceneCapture/Private/CompositionGraphCaptureProtocol.cpp
--- a/Engine/Source/Runtime/MovieSceneCapture/Private/CompositionGraphCaptureProtocol.cpp
+++ b/Engine/Source/Runtime/MovieSceneCapture/Private/CompositionGraphCaptureProtocol.cpp
@@ -174,7 +174,7 @@ bool UCompositionGraphCaptureProtocol::SetupImpl()
 	SceneViewport = InitSettings->SceneViewport;
 
 	FString OverrideRenderPasses;
-	if (FParse::Value(FCommandLine::Get(), TEXT("-CustomRenderPasses="), OverrideRenderPasses))
+	if (FParse::Value(FCommandLine::Get(), TEXT("-CustomRenderPasses="), OverrideRenderPasses, /*bShouldStopOnSeparator*/ false))
 	{
 		OverrideRenderPasses.ParseIntoArray(IncludeRenderPasses.Value, TEXT(","), true);
 	}
