Fixed race condition in movie capture code introduced by 5108105.

  - This change essentially made any calls to CaptureThisFrame capture whichever frame the render thread happened to be rendering at the same time, which might not be the current game thread frame if the capture has been paused.

diff --git a/Engine/Source/Runtime/MovieSceneCapture/Private/FrameGrabber.cpp b/Engine/Source/Runtime/MovieSceneCapture/Private/FrameGrabber.cpp
--- a/Engine/Source/Runtime/MovieSceneCapture/Private/FrameGrabber.cpp
+++ b/Engine/Source/Runtime/MovieSceneCapture/Private/FrameGrabber.cpp
@@ -198,7 +198,6 @@ void FViewportSurfaceReader::ResolveRenderTarget(FViewportSurfaceReader* RenderT
 FFrameGrabber::FFrameGrabber(TSharedRef<FSceneViewport> Viewport, FIntPoint DesiredBufferSize, EPixelFormat InPixelFormat, uint32 NumSurfaces)
 {
 	State = EFrameGrabberState::Inactive;
-	bIsFirstCaptureFrame = false;
 
 	TargetSize = DesiredBufferSize;
 
@@ -266,6 +265,11 @@ FFrameGrabber::~FFrameGrabber()
 	{
 		FSlateApplication::Get().GetRenderer()->OnBackBufferReadyToPresent().Remove(OnBackBufferReadyToPresent);
 	}
+
+	if (OutstandingFrameCount.GetValue() > 0)
+	{
+		FlushRenderingCommands();
+	}
 }
 
 void FFrameGrabber::StartCapturingFrames()
@@ -276,7 +280,6 @@ void FFrameGrabber::StartCapturingFrames()
 	}
 
 	State = EFrameGrabberState::Active;
-	bIsFirstCaptureFrame = true;
 
 	OnBackBufferReadyToPresent = FSlateApplication::Get().GetRenderer()->OnBackBufferReadyToPresent().AddRaw(this, &FFrameGrabber::OnBackBufferReadyToPresentCallback);
 }
@@ -293,19 +296,14 @@ void FFrameGrabber::CaptureThisFrame(FFramePayloadPtr Payload)
 		return;
 	}
 
-	// Callbacks to OnBackBufferReadyToPresentCallback() are coming from the RenderThread, which may still be running when we get here, so we need to wait
-	// until it is done before we increment OutstandingFrameCount here and start capturing frames, otherwise we may end up capturing a frame too early.
-	if (bIsFirstCaptureFrame)
-	{
-		bIsFirstCaptureFrame = false;
-
-		FlushRenderingCommands();
-	}
-
 	OutstandingFrameCount.Increment();
 
-	FScopeLock Lock(&PendingFramePayloadsMutex);
-	PendingFramePayloads.Add(Payload);
+	ENQUEUE_RENDER_COMMAND(FrameGrabber_AppendFramePayload)(
+		[this, Payload](FRHICommandListImmediate& RHICmdList)
+		{
+			this->RenderThread_PendingFramePayloads.Add(Payload);
+		}
+	);
 }
 
 void FFrameGrabber::StopCapturingFrames()
@@ -316,8 +314,6 @@ void FFrameGrabber::StopCapturingFrames()
 	}
 
 	State = EFrameGrabberState::PendingShutdown;
-
-	bIsFirstCaptureFrame = false;
 }
 
 void FFrameGrabber::Shutdown()
@@ -377,17 +373,19 @@ void FFrameGrabber::OnBackBufferReadyToPresentCallback(SWindow& SlateWindow, con
 		return;
 	}
 
+	ensure(IsInRenderingThread());
+
 	FFramePayloadPtr Payload;
 
 	{
-		FScopeLock Lock(&PendingFramePayloadsMutex);
-		if (!PendingFramePayloads.Num())
+		if (!RenderThread_PendingFramePayloads.Num())
 		{
 			// No frames to capture
 			return;
 		}
-		Payload = PendingFramePayloads[0];
-		PendingFramePayloads.RemoveAt(0, 1, false);
+
+		Payload = RenderThread_PendingFramePayloads[0];
+		RenderThread_PendingFramePayloads.RemoveAt(0, 1, false);
 	}
 
 	if (FrameGrabLatency != GFrameGrabberFrameLatency)
diff --git a/Engine/Source/Runtime/MovieSceneCapture/Public/FrameGrabber.h b/Engine/Source/Runtime/MovieSceneCapture/Public/FrameGrabber.h
--- a/Engine/Source/Runtime/MovieSceneCapture/Public/FrameGrabber.h
+++ b/Engine/Source/Runtime/MovieSceneCapture/Public/FrameGrabber.h
@@ -211,9 +211,8 @@ private:
 	/** The total number of frames we are currently waiting on */
 	FThreadSafeCounter OutstandingFrameCount;
 
-	/** Pending frame payloads to be passed with frames captured from slate  */
-	FCriticalSection PendingFramePayloadsMutex;
-	TArray<FFramePayloadPtr> PendingFramePayloads;
+	/** Only to be accessed from the render thread - array of frame payloads to be captured from the rendered slate window sorted first to last. */
+	TArray<FFramePayloadPtr> RenderThread_PendingFramePayloads;
 
 	int32 FrameGrabLatency;
 
@@ -224,8 +223,6 @@ private:
 	};
 	EFrameGrabberState State;
 
-	bool bIsFirstCaptureFrame;
-
 	/** The desired target size to resolve frames to */
 	FIntPoint TargetSize;
 };
