Cherry-pick 6814979: Movie Scene Capture: Force engine show flag path tracing on with -PathTracer

This fixes an issue where if you render from commandline (which launches the editor in PIE), it doesn't enable path tracing.

diff --git a/Engine/Source/Runtime/MovieSceneCapture/Private/MovieSceneCapture.cpp b/Engine/Source/Runtime/MovieSceneCapture/Private/MovieSceneCapture.cpp
--- a/Engine/Source/Runtime/MovieSceneCapture/Private/MovieSceneCapture.cpp
+++ b/Engine/Source/Runtime/MovieSceneCapture/Private/MovieSceneCapture.cpp
@@ -298,6 +298,10 @@ void UMovieSceneCapture::Initialize(TSharedPtr<FSceneViewport> InSceneViewport,
 		if (FParse::Bool(FCommandLine::Get(), TEXT("-PathTracer="), bOverridePathTracer))
 		{
 			Settings.bUsePathTracer = bOverridePathTracer;
+			if (bOverridePathTracer)
+			{
+				InSceneViewport->GetClient()->GetEngineShowFlags()->SetPathTracing(true);
+			}
 		}
 
 		uint16 OverridePathTracerSamplePerPixel;
