Movie Scene Capture: Fix crash on movie rendering from dev-core. IsDestructionThreadSafe returns false.

diff --git a/Engine/Source/Runtime/MovieSceneCapture/Public/MovieSceneCapture.h b/Engine/Source/Runtime/MovieSceneCapture/Public/MovieSceneCapture.h
--- a/Engine/Source/Runtime/MovieSceneCapture/Public/MovieSceneCapture.h
+++ b/Engine/Source/Runtime/MovieSceneCapture/Public/MovieSceneCapture.h
@@ -56,6 +56,10 @@ public:
 	const FMovieSceneCaptureSettings& GetSettings() const override { return Settings; }
 	// End IMovieSceneCaptureInterface
 
+	/*~ Begin UObject interface */
+	virtual bool IsDestructionThreadSafe() const override { return false; }
+	/*~ End UObject interface */
+
 	/** Load save from config helpers */
 	virtual void LoadFromConfig();
 	virtual void SaveToConfig();
