Movie Scene Capture: Get frame rate from the settings so that a custom frame rate can be provided

diff --git a/Engine/Source/Runtime/MovieSceneCapture/Public/MovieSceneCapture.h b/Engine/Source/Runtime/MovieSceneCapture/Public/MovieSceneCapture.h
--- a/Engine/Source/Runtime/MovieSceneCapture/Public/MovieSceneCapture.h
+++ b/Engine/Source/Runtime/MovieSceneCapture/Public/MovieSceneCapture.h
@@ -173,7 +173,7 @@ public:
 	double GetEstimatedCaptureDurationSeconds() const { return 0.0; }
 
 
-	virtual FFrameRate GetCaptureFrameRate() const { return Settings.FrameRate; }
+	virtual FFrameRate GetCaptureFrameRate() const { return Settings.GetFrameRate(); }
 	virtual const ICaptureStrategy& GetCaptureStrategy() const { return *CaptureStrategy; }
 	// ~ICaptureProtocolHost interface
 
