Movie Scene Capture: Remove unused codec

diff --git a/Engine/Source/Runtime/MovieSceneCapture/Private/VideoCaptureProtocol.cpp b/Engine/Source/Runtime/MovieSceneCapture/Private/VideoCaptureProtocol.cpp
--- a/Engine/Source/Runtime/MovieSceneCapture/Private/VideoCaptureProtocol.cpp
+++ b/Engine/Source/Runtime/MovieSceneCapture/Private/VideoCaptureProtocol.cpp
@@ -39,7 +39,6 @@ void UVideoCaptureProtocol::ConditionallyCreateWriter()
 	Options.OutputFilename = MoveTemp(VideoFilename);
 	Options.CaptureFramerateNumerator = CaptureHost->GetCaptureFrameRate().Numerator;
 	Options.CaptureFramerateDenominator = CaptureHost->GetCaptureFrameRate().Denominator;
-	Options.CodecName = VideoCodec;
 	Options.bSynchronizeFrames = CaptureHost->GetCaptureStrategy().ShouldSynchronizeFrames();
 	Options.Width = InitSettings->DesiredSize.X;
 	Options.Height = InitSettings->DesiredSize.Y;
diff --git a/Engine/Source/Runtime/MovieSceneCapture/Public/Protocols/VideoCaptureProtocol.h b/Engine/Source/Runtime/MovieSceneCapture/Public/Protocols/VideoCaptureProtocol.h
--- a/Engine/Source/Runtime/MovieSceneCapture/Public/Protocols/VideoCaptureProtocol.h
+++ b/Engine/Source/Runtime/MovieSceneCapture/Public/Protocols/VideoCaptureProtocol.h
@@ -31,9 +31,6 @@ public:
 	UPROPERTY(config, EditAnywhere, Category=VideoSettings, meta=(ClampMin=1, ClampMax=100, EditCondition=bUseCompression))
 	float CompressionQuality;
 
-	UPROPERTY(config, EditAnywhere, Category=VideoSettings, AdvancedDisplay)
-	FString VideoCodec;
-
 public:
 	virtual bool SetupImpl() override;
 	virtual void FinalizeImpl() override;
