Renderpass updates.

Framegrabber. Self-contained.

diff --git a/Engine/Source/Runtime/MovieSceneCapture/Private/FrameGrabber.cpp b/Engine/Source/Runtime/MovieSceneCapture/Private/FrameGrabber.cpp
--- a/Engine/Source/Runtime/MovieSceneCapture/Private/FrameGrabber.cpp
+++ b/Engine/Source/Runtime/MovieSceneCapture/Private/FrameGrabber.cpp
@@ -150,65 +150,62 @@ void FViewportSurfaceReader::ResolveRenderTarget(FViewportSurfaceReader* RenderT
 
 		const FSceneRenderTargetItem& DestRenderTarget = ResampleTexturePooledRenderTarget->GetRenderTargetItem();
 
-		SetRenderTarget(RHICmdList, DestRenderTarget.TargetableTexture, FTextureRHIRef());
-		RHICmdList.SetViewport(0, 0, 0.0f, TargetSize.X, TargetSize.Y, 1.0f);
+		FRHIRenderPassInfo RPInfo(DestRenderTarget.TargetableTexture, ERenderTargetActions::Load_Store, ReadbackTexture);
+		RHICmdList.BeginRenderPass(RPInfo, TEXT("FrameGrabberResolveRenderTarget"));
+		{
+			RHICmdList.SetViewport(0, 0, 0.0f, TargetSize.X, TargetSize.Y, 1.0f);
 
-		FGraphicsPipelineStateInitializer GraphicsPSOInit;
-		RHICmdList.ApplyCachedRenderTargets(GraphicsPSOInit);
-		GraphicsPSOInit.BlendState = TStaticBlendState<>::GetRHI();
-		GraphicsPSOInit.RasterizerState = TStaticRasterizerState<>::GetRHI();
-		GraphicsPSOInit.DepthStencilState = TStaticDepthStencilState<false,CF_Always>::GetRHI();
+			FGraphicsPipelineStateInitializer GraphicsPSOInit;
+			RHICmdList.ApplyCachedRenderTargets(GraphicsPSOInit);
+			GraphicsPSOInit.BlendState = TStaticBlendState<>::GetRHI();
+			GraphicsPSOInit.RasterizerState = TStaticRasterizerState<>::GetRHI();
+			GraphicsPSOInit.DepthStencilState = TStaticDepthStencilState<false, CF_Always>::GetRHI();
 
-		const ERHIFeatureLevel::Type FeatureLevel = GMaxRHIFeatureLevel;
-		
-		TShaderMap<FGlobalShaderType>* ShaderMap = GetGlobalShaderMap(FeatureLevel);
-		TShaderMapRef<FScreenVS> VertexShader(ShaderMap);
-		TShaderMapRef<FScreenPS> PixelShader(ShaderMap);
+			const ERHIFeatureLevel::Type FeatureLevel = GMaxRHIFeatureLevel;
 
-		GraphicsPSOInit.BoundShaderState.VertexDeclarationRHI = RendererModule->GetFilterVertexDeclaration().VertexDeclarationRHI;
-		GraphicsPSOInit.BoundShaderState.VertexShaderRHI = GETSAFERHISHADER_VERTEX(*VertexShader);
-		GraphicsPSOInit.BoundShaderState.PixelShaderRHI = GETSAFERHISHADER_PIXEL(*PixelShader);
-		GraphicsPSOInit.PrimitiveType = PT_TriangleList;
+			TShaderMap<FGlobalShaderType>* ShaderMap = GetGlobalShaderMap(FeatureLevel);
+			TShaderMapRef<FScreenVS> VertexShader(ShaderMap);
+			TShaderMapRef<FScreenPS> PixelShader(ShaderMap);
 
-		SetGraphicsPipelineState(RHICmdList, GraphicsPSOInit);
+			GraphicsPSOInit.BoundShaderState.VertexDeclarationRHI = RendererModule->GetFilterVertexDeclaration().VertexDeclarationRHI;
+			GraphicsPSOInit.BoundShaderState.VertexShaderRHI = GETSAFERHISHADER_VERTEX(*VertexShader);
+			GraphicsPSOInit.BoundShaderState.PixelShaderRHI = GETSAFERHISHADER_PIXEL(*PixelShader);
+			GraphicsPSOInit.PrimitiveType = PT_TriangleList;
 
-		FTexture2DRHIRef SourceBackBuffer = RHICmdList.GetViewportBackBuffer(ViewportRHI);
+			SetGraphicsPipelineState(RHICmdList, GraphicsPSOInit);
 
-		const bool bIsSourceBackBufferSameAsWindowSize = SourceBackBuffer->GetSizeX() == WindowSize.X && SourceBackBuffer->GetSizeY() == WindowSize.Y;
-		const bool bIsSourceBackBufferSameAsTargetSize = TargetSize.X == SourceBackBuffer->GetSizeX() && TargetSize.Y == SourceBackBuffer->GetSizeY();
+			FTexture2DRHIRef SourceBackBuffer = RHICmdList.GetViewportBackBuffer(ViewportRHI);
 
-		if (bIsSourceBackBufferSameAsWindowSize || bIsSourceBackBufferSameAsTargetSize)
-		{
-			PixelShader->SetParameters(RHICmdList, TStaticSamplerState<SF_Point>::GetRHI(), SourceBackBuffer);
-		}
-		else
-		{
-			PixelShader->SetParameters(RHICmdList, TStaticSamplerState<SF_Bilinear>::GetRHI(), SourceBackBuffer);
-		}
+			const bool bIsSourceBackBufferSameAsWindowSize = SourceBackBuffer->GetSizeX() == WindowSize.X && SourceBackBuffer->GetSizeY() == WindowSize.Y;
+			const bool bIsSourceBackBufferSameAsTargetSize = TargetSize.X == SourceBackBuffer->GetSizeX() && TargetSize.Y == SourceBackBuffer->GetSizeY();
 
-		float U = float(CaptureRect.Min.X) / float(SourceBackBuffer->GetSizeX());
-		float V = float(CaptureRect.Min.Y) / float(SourceBackBuffer->GetSizeY());
-		float SizeU = float(CaptureRect.Max.X) / float(SourceBackBuffer->GetSizeX()) - U;
-		float SizeV = float(CaptureRect.Max.Y) / float(SourceBackBuffer->GetSizeY()) - V;
-
-		RendererModule->DrawRectangle(
-			RHICmdList,
-			0, 0,									// Dest X, Y
-			TargetSize.X,							// Dest Width
-			TargetSize.Y,							// Dest Height
-			U, V,									// Source U, V
-			SizeU, SizeV,							// Source USize, VSize
-			CaptureRect.Max - CaptureRect.Min,		// Target buffer size
-			FIntPoint(1, 1),						// Source texture size
-			*VertexShader,
-			EDRF_Default);
-
-		// Asynchronously copy render target from GPU to CPU
-		RHICmdList.CopyToResolveTarget(
-			DestRenderTarget.TargetableTexture,
-			ReadbackTexture,
-			FResolveParams());
+			if (bIsSourceBackBufferSameAsWindowSize || bIsSourceBackBufferSameAsTargetSize)
+			{
+				PixelShader->SetParameters(RHICmdList, TStaticSamplerState<SF_Point>::GetRHI(), SourceBackBuffer);
+			}
+			else
+			{
+				PixelShader->SetParameters(RHICmdList, TStaticSamplerState<SF_Bilinear>::GetRHI(), SourceBackBuffer);
+			}
 
+			float U = float(CaptureRect.Min.X) / float(SourceBackBuffer->GetSizeX());
+			float V = float(CaptureRect.Min.Y) / float(SourceBackBuffer->GetSizeY());
+			float SizeU = float(CaptureRect.Max.X) / float(SourceBackBuffer->GetSizeX()) - U;
+			float SizeV = float(CaptureRect.Max.Y) / float(SourceBackBuffer->GetSizeY()) - V;
+
+			RendererModule->DrawRectangle(
+				RHICmdList,
+				0, 0,									// Dest X, Y
+				TargetSize.X,							// Dest Width
+				TargetSize.Y,							// Dest Height
+				U, V,									// Source U, V
+				SizeU, SizeV,							// Source USize, VSize
+				CaptureRect.Max - CaptureRect.Min,		// Target buffer size
+				FIntPoint(1, 1),						// Source texture size
+				*VertexShader,
+				EDRF_Default);
+		}
+		RHICmdList.EndRenderPass();
 
 		if (InRenderToReadback)
 		{
