Composure: Fixed crash on composure export when exporting as jpg

  - The errors are now outputted to the message log

diff --git a/Engine/Source/Runtime/MovieSceneCapture/Private/UserDefinedCaptureProtocol.cpp b/Engine/Source/Runtime/MovieSceneCapture/Private/UserDefinedCaptureProtocol.cpp
--- a/Engine/Source/Runtime/MovieSceneCapture/Private/UserDefinedCaptureProtocol.cpp
+++ b/Engine/Source/Runtime/MovieSceneCapture/Private/UserDefinedCaptureProtocol.cpp
@@ -11,6 +11,9 @@
 #include "Slate/SceneViewport.h"
 #include "ImagePixelData.h"
 #include "Engine/Engine.h"
+#include "Logging/MessageLog.h"
+
+#define LOCTEXT_NAMESPACE "UserDefinedImageCaptureProtocol"
 
 struct FCaptureProtocolFrameData : IFramePayload
 {
@@ -65,6 +68,26 @@ private:
 	TWeakObjectPtr<UUserDefinedCaptureProtocol> WeakProtocol;
 };
 
+FString FCapturedPixelsID::ToString() const
+{
+	FString Name;
+	for (const TTuple<FName, FName>& Pair : Identifiers)
+	{
+		if (Name.Len() > 0)
+		{
+			Name += TEXT(",");
+		}
+
+		Name += Pair.Key.ToString();
+		if (Pair.Value != NAME_None)
+		{
+			Name += TEXT(":");
+			Name += Pair.Value.ToString();
+		}
+	}
+	return Name.Len() > 0 ? Name : TEXT("<none>");
+}
+
 
 UUserDefinedCaptureProtocol::UUserDefinedCaptureProtocol(const FObjectInitializer& ObjInit)
 	: Super(ObjInit)
@@ -371,15 +394,26 @@ void UUserDefinedImageCaptureProtocol::WriteImageToDisk(const FCapturedPixels& P
 	{
 		return;
 	}
-
-	TUniquePtr<FImageWriteTask> ImageTask = MakeUnique<FImageWriteTask>();
-
-	// If the pixels are FColors, and this is the final pixels buffer, and we're writing PNG, always write out full alpha
-	if (PixelData.ImageData->GetType() == EImagePixelType::Color && ImageTask->Format == EImageFormat::PNG && StreamID.Identifiers.OrderIndependentCompareEqual(FinalPixelsID.Identifiers))
+	else if (PixelData.ImageData->GetBitDepth() != 8)
 	{
-		ImageTask->PixelPreProcessors.Add(TAsyncAlphaWrite<FColor>(255));
+		if (Format == EDesiredImageFormat::BMP)
+		{
+			FMessageLog("PIE")
+			.Warning(FText::Format(LOCTEXT("InvalidBMPExport", "Unable to write the specified render target (stream '{0}' is {1}bit) as BMP. BMPs must be supplied 8bit render targets."),
+				FText::FromString(StreamID.ToString()), FText::AsNumber(PixelData.ImageData->GetBitDepth())));
+			return;
+		}
+		else if (Format == EDesiredImageFormat::JPG)
+		{
+			FMessageLog("PIE")
+				.Warning(FText::Format(LOCTEXT("InvalidJPGExport", "Unable to write the specified render target (stream '{0}' is {1}bit) as JPG. JPGs must be supplied 8bit render targets."),
+					FText::FromString(StreamID.ToString()), FText::AsNumber(PixelData.ImageData->GetBitDepth())));
+			return;
+		}
 	}
 
+	TUniquePtr<FImageWriteTask> ImageTask = MakeUnique<FImageWriteTask>();
+
 	// Cache the buffer ID so we generate the correct filename
 	CurrentStreamID = &StreamID;
 
@@ -390,6 +424,12 @@ void UUserDefinedImageCaptureProtocol::WriteImageToDisk(const FCapturedPixels& P
 
 	CurrentStreamID = nullptr;
 
+	// If the pixels are FColors, and this is the final pixels buffer, and we're writing PNG, always write out full alpha
+	if (PixelData.ImageData->GetType() == EImagePixelType::Color && ImageTask->Format == EImageFormat::PNG && StreamID.Identifiers.OrderIndependentCompareEqual(FinalPixelsID.Identifiers))
+	{
+		ImageTask->PixelPreProcessors.Add(TAsyncAlphaWrite<FColor>(255));
+	}
+
 	if (Format == EDesiredImageFormat::EXR)
 	{
 		ImageTask->CompressionQuality = bEnableCompression ? (int32)EImageCompressionQuality::Default : (int32)EImageCompressionQuality::Uncompressed;
@@ -424,4 +464,6 @@ void UUserDefinedImageCaptureProtocol::WriteImageToDisk(const FCapturedPixels& P
 void UUserDefinedImageCaptureProtocol::OnFileWritten()
 {
 	--NumOutstandingOperations;
-}
\ No newline at end of file
+}
+
+#undef LOCTEXT_NAMESPACE
\ No newline at end of file
diff --git a/Engine/Source/Runtime/MovieSceneCapture/Public/Protocols/UserDefinedCaptureProtocol.h b/Engine/Source/Runtime/MovieSceneCapture/Public/Protocols/UserDefinedCaptureProtocol.h
--- a/Engine/Source/Runtime/MovieSceneCapture/Public/Protocols/UserDefinedCaptureProtocol.h
+++ b/Engine/Source/Runtime/MovieSceneCapture/Public/Protocols/UserDefinedCaptureProtocol.h
@@ -20,6 +20,9 @@ struct FCapturedPixelsID
 {
 	GENERATED_BODY()
 
+	/** Convert this strream ID to a string */
+	FString ToString() const;
+
 	/** Map of identifiers to their values for this ID. */
 	UPROPERTY(EditAnywhere, BlueprintReadWrite, Category=Capture)
 	TMap<FName, FName> Identifiers;
