Removed redundant code (this memory overwrite was diagnosed as a call to LoadModule on a background thread in completely unrelated code)

diff --git a/Engine/Source/Runtime/MovieSceneCapture/Private/FrameGrabber.cpp b/Engine/Source/Runtime/MovieSceneCapture/Private/FrameGrabber.cpp
--- a/Engine/Source/Runtime/MovieSceneCapture/Private/FrameGrabber.cpp
+++ b/Engine/Source/Runtime/MovieSceneCapture/Private/FrameGrabber.cpp
@@ -109,19 +109,6 @@ void FViewportSurfaceReader::ResolveRenderTarget(FViewportSurfaceReader* RenderT
 	{
 		FRHICommandListImmediate& RHICmdList = GetImmediateCommandList_ForRenderCommand();
 
-		// @todo: JIRA UE-41879 and UE-43829. If any of these ensures go off, something has overwritten the memory for this render command (buffer underflow/overflow?)
-		bool bMemoryTrample = !ensureMsgf(RendererModule, TEXT("RendererModule has become null. This indicates a memory trample.")) ||
-			!ensureMsgf(RendererModule == RendererModuleDebug, TEXT("RendererModule and RendererModuleDebug are not equal (0x%016x != 0x%016x). This indicates a memory trample."), (void*)RendererModule, (void*)RendererModuleDebug) ||
-			!ensureMsgf(MemoryGuard1 == 0xaffec7ed, TEXT("Memory guard 1 is now 0x%08x, expected 0xaffec7ed."), MemoryGuard1) || //-V547
-			!ensureMsgf(MemoryGuard2 == 0xaffec7ed, TEXT("Memory guard 2 is now 0x%08x, expected 0xaffec7ed."), MemoryGuard2); //-V547
-
-		if (bMemoryTrample)
-		{
-			// In the hope that 'this' is still ok, triggering the event will prevent a deadlock. If it's not ok, this may crash, but it was going to crash anyway
-			RenderToReadback->AvailableEvent->Trigger();
-			return;
-		}
-
 		const FIntPoint TargetSize(ReadbackTexture->GetSizeX(), ReadbackTexture->GetSizeY());
 
 		FPooledRenderTargetDesc OutputDesc = FPooledRenderTargetDesc::Create2DDesc(
@@ -132,19 +119,6 @@ void FViewportSurfaceReader::ResolveRenderTarget(FViewportSurfaceReader* RenderT
 			TexCreate_RenderTargetable,
 			false);
 
-		// @todo: JIRA UE-41879 and UE-43829. If any of these ensures go off, something has overwritten the memory for this render command (buffer underflow/overflow?)
-		bMemoryTrample = !ensureMsgf(RendererModule, TEXT("RendererModule has become null. This indicates a memory trample.")) ||
-			!ensureMsgf(RendererModule == RendererModuleDebug, TEXT("RendererModule and RendererModuleDebug are not equal (0x%16x != 0x%16x). This indicates a memory trample."), (void*)RendererModule, (void*)RendererModuleDebug) ||
-			!ensureMsgf(MemoryGuard1 == 0xaffec7ed, TEXT("Memory guard 1 is now 0x%08x, expected 0xaffec7ed."), MemoryGuard1) ||
-			!ensureMsgf(MemoryGuard2 == 0xaffec7ed, TEXT("Memory guard 2 is now 0x%08x, expected 0xaffec7ed."), MemoryGuard2);
-
-		if (bMemoryTrample)
-		{
-			// In the hope that 'this' is still ok, triggering the event will prevent a deadlock. If it's not ok, this may crash, but it was going to crash anyway
-			RenderToReadback->AvailableEvent->Trigger();
-			return;
-		}
-
 		TRefCountPtr<IPooledRenderTarget> ResampleTexturePooledRenderTarget;
 		RendererModule->RenderTargetPoolFindFreeElement(RHICmdList, OutputDesc, ResampleTexturePooledRenderTarget, TEXT("ResampleTexture"));
 		check(ResampleTexturePooledRenderTarget);
