Fixed crash during rendering caused by the threaded access to a non-thread-safe TSharedPtr<SWindow>

diff --git a/Engine/Source/Runtime/MovieSceneCapture/Private/FrameGrabber.cpp b/Engine/Source/Runtime/MovieSceneCapture/Private/FrameGrabber.cpp
--- a/Engine/Source/Runtime/MovieSceneCapture/Private/FrameGrabber.cpp
+++ b/Engine/Source/Runtime/MovieSceneCapture/Private/FrameGrabber.cpp
@@ -203,6 +203,7 @@ FFrameGrabber::FFrameGrabber(TSharedRef<FSceneViewport> Viewport, FIntPoint Desi
 	TargetSize = DesiredBufferSize;
 
 	CurrentFrameIndex = 0;
+	TargetWindowPtr = nullptr;
 
 	check(NumSurfaces != 0);
 
@@ -216,7 +217,7 @@ FFrameGrabber::FFrameGrabber(TSharedRef<FSceneViewport> Viewport, FIntPoint Desi
 		TSharedPtr<SWindow> Window = FSlateApplication::Get().FindWidgetWindow(ViewportWidget.ToSharedRef());
 		if (Window.IsValid())
 		{
-			CaptureWindow = Window;
+			TargetWindowPtr = Window.Get();
 			FGeometry InnerWindowGeometry = Window->GetWindowGeometryInWindow();
 			
 			// Find the widget path relative to the window
@@ -371,8 +372,7 @@ TArray<FCapturedFrameData> FFrameGrabber::GetCapturedFrames()
 void FFrameGrabber::OnBackBufferReadyToPresentCallback(SWindow& SlateWindow, const FTexture2DRHIRef& BackBuffer)
 {
 	// We only care about our own Slate window
-	TSharedPtr<SWindow> Window = CaptureWindow.Pin();
-	if (Window.Get() != &SlateWindow)
+	if (&SlateWindow != TargetWindowPtr)
 	{
 		return;
 	}
diff --git a/Engine/Source/Runtime/MovieSceneCapture/Public/FrameGrabber.h b/Engine/Source/Runtime/MovieSceneCapture/Public/FrameGrabber.h
--- a/Engine/Source/Runtime/MovieSceneCapture/Public/FrameGrabber.h
+++ b/Engine/Source/Runtime/MovieSceneCapture/Public/FrameGrabber.h
@@ -176,8 +176,12 @@ private:
 	FFrameGrabber(const FFrameGrabber&);
 	FFrameGrabber& operator=(const FFrameGrabber&);
 
-	/** The window we'll capture, and the capture rectangle */
-	TWeakPtr<SWindow> CaptureWindow;
+	/**
+	 * Pointer to the window we want to capture.
+	 * Only held for comparison inside OnBackBufferReadyToPresentCallback - never to be dereferenced or cast to an SWindow.
+	 * Held as a raw pointer to ensure that no referenc counting occurs from the background thread in OnBackBufferReadyToPresentCallback.
+	 */
+	void* TargetWindowPtr;
 
 	/** Delegate handle for the OnBackBufferReadyToPresent event */
 	FDelegateHandle OnBackBufferReadyToPresent;
