Movie Scene Capture: Tweaks to custom frame rate settings

diff --git a/Engine/Source/Runtime/MovieSceneCapture/Private/MovieSceneCapture.cpp b/Engine/Source/Runtime/MovieSceneCapture/Private/MovieSceneCapture.cpp
--- a/Engine/Source/Runtime/MovieSceneCapture/Private/MovieSceneCapture.cpp
+++ b/Engine/Source/Runtime/MovieSceneCapture/Private/MovieSceneCapture.cpp
@@ -380,10 +380,13 @@ void UMovieSceneCapture::Initialize(TSharedPtr<FSceneViewport> InSceneViewport,
 		FString FrameRateOverrideString;
 		if ( FParse::Value( FCommandLine::Get(), TEXT( "-MovieFrameRate=" ), FrameRateOverrideString ) )
 		{
-			Settings.bUseCustomFrameRate = true;
 			if (!TryParseString(Settings.CustomFrameRate, *FrameRateOverrideString))
 			{
-				UE_LOG(LogMovieSceneCapture, Error, TEXT("Unrecognized capture frame rate: %s."), *FrameRateOverrideString);
+				UE_LOG(LogMovieSceneCapture, Error, TEXT("Unrecognized capture frame rate: %s. Defaulting to sequence frame rate."), *FrameRateOverrideString);
+			}
+			else
+			{
+				Settings.bUseCustomFrameRate = true;
 			}
 		}
 	}
