Movie Scene Capture: Changed from Warning to Display

diff --git a/Engine/Source/Runtime/MovieSceneCapture/Private/CompositionGraphCaptureProtocol.cpp b/Engine/Source/Runtime/MovieSceneCapture/Private/CompositionGraphCaptureProtocol.cpp
--- a/Engine/Source/Runtime/MovieSceneCapture/Private/CompositionGraphCaptureProtocol.cpp
+++ b/Engine/Source/Runtime/MovieSceneCapture/Private/CompositionGraphCaptureProtocol.cpp
@@ -11,6 +11,7 @@
 #include "SceneViewExtension.h"
 #include "Materials/Material.h"
 #include "BufferVisualizationData.h"
+#include "MovieSceneCaptureModule.h"
 #include "MovieSceneCaptureSettings.h"
 
 struct FFrameCaptureViewExtension : public FSceneViewExtensionBase
@@ -212,7 +213,7 @@ void UCompositionGraphCaptureProtocol::OnLoadConfigImpl(FMovieSceneCaptureSettin
 		OutputFormat.Append(TEXT(".{frame}"));
 
 		InSettings.OutputFormat = OutputFormat;
-		UE_LOG(LogTemp, Warning, TEXT("Automatically appended .{frame} to the format string as specified format string did not provide a way to differentiate between frames via {frame} or {shot_frame}!"));
+		UE_LOG(LogMovieSceneCapture, Display, TEXT("Automatically appended .{frame} to the format string as specified format string did not provide a way to differentiate between frames via {frame} or {shot_frame}!"));
 	}
 
 	// Add {material} if it doesn't already exist
diff --git a/Engine/Source/Runtime/MovieSceneCapture/Private/ImageSequenceProtocol.cpp b/Engine/Source/Runtime/MovieSceneCapture/Private/ImageSequenceProtocol.cpp
--- a/Engine/Source/Runtime/MovieSceneCapture/Private/ImageSequenceProtocol.cpp
+++ b/Engine/Source/Runtime/MovieSceneCapture/Private/ImageSequenceProtocol.cpp
@@ -12,6 +12,7 @@
 #include "Async/Future.h"
 #include "Async/Async.h"
 #include "Templates/Casts.h"
+#include "MovieSceneCaptureModule.h"
 #include "MovieSceneCaptureSettings.h"
 #include "ImageWriteQueue.h"
 
@@ -39,7 +40,7 @@ void UImageSequenceProtocol::OnLoadConfigImpl(FMovieSceneCaptureSettings& InSett
 		OutputFormat.Append(TEXT(".{frame}"));
 		InSettings.OutputFormat = OutputFormat;
 
-		UE_LOG(LogTemp, Warning, TEXT("Automatically appended .{frame} to the format string as specified format string did not provide a way to differentiate between frames via {frame} or {shot_frame}!"));
+		UE_LOG(LogMovieSceneCapture, Display, TEXT("Automatically appended .{frame} to the format string as specified format string did not provide a way to differentiate between frames via {frame} or {shot_frame}!"));
 	}
 
 	Super::OnLoadConfigImpl(InSettings);
diff --git a/Engine/Source/Runtime/MovieSceneCapture/Private/UserDefinedCaptureProtocol.cpp b/Engine/Source/Runtime/MovieSceneCapture/Private/UserDefinedCaptureProtocol.cpp
--- a/Engine/Source/Runtime/MovieSceneCapture/Private/UserDefinedCaptureProtocol.cpp
+++ b/Engine/Source/Runtime/MovieSceneCapture/Private/UserDefinedCaptureProtocol.cpp
@@ -7,6 +7,7 @@
 #include "Async/Async.h"
 #include "Engine/Texture.h"
 #include "UnrealClient.h"
+#include "MovieSceneCaptureModule.h"
 #include "MovieSceneCaptureSettings.h"
 #include "Slate/SceneViewport.h"
 #include "ImagePixelData.h"
@@ -325,7 +326,7 @@ void UUserDefinedImageCaptureProtocol::OnLoadConfigImpl(FMovieSceneCaptureSettin
 		OutputFormat.Append(TEXT(".{frame}"));
 
 		InSettings.OutputFormat = OutputFormat;
-		UE_LOG(LogTemp, Warning, TEXT("Automatically appended .{frame} to the format string as specified format string did not provide a way to differentiate between frames via {frame} or {shot_frame}!"));
+		UE_LOG(LogMovieSceneCapture, Display, TEXT("Automatically appended .{frame} to the format string as specified format string did not provide a way to differentiate between frames via {frame} or {shot_frame}!"));
 	}
 }
 
