Setting Movie Scene Capture to BlueprintType

diff --git a/Engine/Source/Runtime/MovieSceneCapture/Public/MovieSceneCapture.h b/Engine/Source/Runtime/MovieSceneCapture/Public/MovieSceneCapture.h
--- a/Engine/Source/Runtime/MovieSceneCapture/Public/MovieSceneCapture.h
+++ b/Engine/Source/Runtime/MovieSceneCapture/Public/MovieSceneCapture.h
@@ -33,7 +33,7 @@ struct FCachedMetrics
 };
 
 /** Class responsible for capturing scene data */
-UCLASS(config=EditorPerProjectUserSettings, PerObjectConfig)
+UCLASS(config=EditorPerProjectUserSettings, PerObjectConfig, BlueprintType)
 class MOVIESCENECAPTURE_API UMovieSceneCapture : public UObject, public IMovieSceneCaptureInterface, public ICaptureProtocolHost
 {
 public:
