Sequencer: Removed transient from FMovieSceneCaptureSettings::FrameRate to ensure it gets serialized into the capture manifest

diff --git a/Engine/Source/Runtime/MovieSceneCapture/Public/MovieSceneCaptureSettings.h b/Engine/Source/Runtime/MovieSceneCapture/Public/MovieSceneCaptureSettings.h
--- a/Engine/Source/Runtime/MovieSceneCapture/Public/MovieSceneCaptureSettings.h
+++ b/Engine/Source/Runtime/MovieSceneCapture/Public/MovieSceneCaptureSettings.h
@@ -77,7 +77,7 @@ struct MOVIESCENECAPTURE_API FMovieSceneCaptureSettings
 	uint8 ZeroPadFrameNumbers;
 
 	/** The sequence's frame rate at which to capture if "Use Custom Frame Rate" is not enabled */
-	UPROPERTY(transient, VisibleAnywhere, BlueprintReadWrite, Category = CaptureSettings, meta = (ClampMin = 1, UIMin = 1, ClampMax = 200, UIMax = 200))
+	UPROPERTY(VisibleAnywhere, BlueprintReadWrite, Category = CaptureSettings, meta = (ClampMin = 1, UIMin = 1, ClampMax = 200, UIMax = 200))
 	FFrameRate FrameRate;
 
 	/** Specify using the custom frame rate as opposed to the sequence's display rate */
