Attempting fix for movie rendering causing negative delta times after completing and reverting off of fixed timestep.

diff --git a/Engine/Source/Runtime/MovieSceneCapture/Private/MovieSceneCapture.cpp b/Engine/Source/Runtime/MovieSceneCapture/Private/MovieSceneCapture.cpp
--- a/Engine/Source/Runtime/MovieSceneCapture/Private/MovieSceneCapture.cpp
+++ b/Engine/Source/Runtime/MovieSceneCapture/Private/MovieSceneCapture.cpp
@@ -826,6 +826,15 @@ void FFixedTimeStepCaptureStrategy::OnInitialize()
 void FFixedTimeStepCaptureStrategy::OnStop()
 {
 	FApp::SetUseFixedTimeStep(false);
+
+	// This can cause an issue with negative delta times in the engine. Instead of changing
+	// that rather sensitive code, we're going to just override the time that calculation is based off of, so it never comes up with a negative number.
+	if(FApp::GetCurrentTime() > FPlatformTime::Seconds())
+	{
+		UE_LOG(LogTemp, Warning, TEXT("FApp::CurrentTime() is ahead of platform time due to Fixed Timestep. Fixing. AppCurrentTime: %f PlatformCurrentTime: %f"), FApp::GetCurrentTime(), FPlatformTime::Seconds());
+		double CurrentRealTime = FPlatformTime::Seconds();
+		FApp::SetCurrentTime(CurrentRealTime);
+	}
 }
 
 bool FFixedTimeStepCaptureStrategy::ShouldPresent(double CurrentTimeSeconds, uint32 FrameIndex) const
