CaptureGammut command line argument now parses names of enum entries as well as numbers.

diff --git a/Engine/Source/Runtime/MovieSceneCapture/Private/ImageSequenceProtocol.cpp b/Engine/Source/Runtime/MovieSceneCapture/Private/ImageSequenceProtocol.cpp
--- a/Engine/Source/Runtime/MovieSceneCapture/Private/ImageSequenceProtocol.cpp
+++ b/Engine/Source/Runtime/MovieSceneCapture/Private/ImageSequenceProtocol.cpp
@@ -8,6 +8,7 @@
 #include "Misc/ScopeLock.h"
 #include "Misc/FeedbackContext.h"
 #include "Misc/ScopedSlowTask.h"
+#include "Misc/DefaultValueHelper.h"
 #include "Modules/ModuleManager.h"
 #include "Async/Future.h"
 #include "Async/Async.h"
@@ -179,8 +180,24 @@ bool UImageSequenceProtocol_EXR::SetupImpl()
 {
 	{
 		int32 OverrideCaptureGamut = (int32)CaptureGamut;
-		FParse::Value(FCommandLine::Get(), TEXT("-CaptureGamut="), OverrideCaptureGamut);
-		CaptureGamut = (EHDRCaptureGamut)OverrideCaptureGamut;
+		FString CaptureGamutString;
+
+		if (FParse::Value(FCommandLine::Get(), TEXT("-CaptureGamut="), CaptureGamutString))
+		{
+			if (!FDefaultValueHelper::ParseInt(CaptureGamutString, OverrideCaptureGamut))
+			{
+				OverrideCaptureGamut = StaticEnum<EHDRCaptureGamut>()->GetValueByName(FName(*CaptureGamutString));
+			}
+			// Invalid CaptureGamut will crash (see UImageSequenceProtocol_EXR::AddFormatMappingsImpl), so only set if valid.
+			if (OverrideCaptureGamut > INDEX_NONE && OverrideCaptureGamut < EHDRCaptureGamut::HCGM_MAX)
+			{
+				CaptureGamut = (EHDRCaptureGamut)OverrideCaptureGamut;
+			}
+			else
+			{
+				UE_LOG(LogMovieSceneCapture, Warning, TEXT("The value for the command -CaptureGamut is invalid, using default value instead!"))
+			}
+		}
 	}
 
 	int32 HDRCompressionQuality = 0;
