Movie Scene Capture: Add -PostProcessingMaterial commandline argument

diff --git a/Engine/Source/Runtime/MovieSceneCapture/Private/CompositionGraphCaptureProtocol.cpp b/Engine/Source/Runtime/MovieSceneCapture/Private/CompositionGraphCaptureProtocol.cpp
--- a/Engine/Source/Runtime/MovieSceneCapture/Private/CompositionGraphCaptureProtocol.cpp
+++ b/Engine/Source/Runtime/MovieSceneCapture/Private/CompositionGraphCaptureProtocol.cpp
@@ -186,6 +186,11 @@ bool UCompositionGraphCaptureProtocol::SetupImpl()
 	FParse::Bool(FCommandLine::Get(), TEXT("-CaptureFramesInHDR="), bCaptureFramesInHDR);
 	FParse::Bool(FCommandLine::Get(), TEXT("-DisableScreenPercentage="), bDisableScreenPercentage);
 
+	FString OverridePostProcessingMaterial;
+	if (FParse::Value(FCommandLine::Get(), TEXT("-PostProcessingMaterial="), OverridePostProcessingMaterial, /*bShouldStopOnSeparator*/ false))
+	{
+		PostProcessingMaterial.SetPath(OverridePostProcessingMaterial);
+	}
 	PostProcessingMaterialPtr = Cast<UMaterialInterface>(PostProcessingMaterial.TryLoad());
 	ViewExtension = FSceneViewExtensions::NewExtension<FFrameCaptureViewExtension>(IncludeRenderPasses.Value, bCaptureFramesInHDR, HDRCompressionQuality, OverrideCaptureGamut, PostProcessingMaterialPtr, bDisableScreenPercentage);
 
