Compilation fixes

diff --git a/Engine/Platforms/Hololens/Source/Runtime/Windows/D3D11RHI/Private/HoloLensD3D11Viewport.cpp b/Engine/Platforms/Hololens/Source/Runtime/Windows/D3D11RHI/Private/HoloLensD3D11Viewport.cpp
--- a/Engine/Platforms/Hololens/Source/Runtime/Windows/D3D11RHI/Private/HoloLensD3D11Viewport.cpp
+++ b/Engine/Platforms/Hololens/Source/Runtime/Windows/D3D11RHI/Private/HoloLensD3D11Viewport.cpp
@@ -239,4 +239,8 @@ void FD3D11Viewport::ResetSwapChainInternal(bool bIgnoreFocus)
 			}
 		}
 	}
-}
\ No newline at end of file
+}
+
+void FD3D11DynamicRHI::RHIGetDisplaysInformation(FDisplayInformationArray& OutDisplayInformation)
+{
+}
diff --git a/Engine/Source/Runtime/MovieSceneCapture/Private/CompositionGraphCaptureProtocol.cpp b/Engine/Source/Runtime/MovieSceneCapture/Private/CompositionGraphCaptureProtocol.cpp
--- a/Engine/Source/Runtime/MovieSceneCapture/Private/CompositionGraphCaptureProtocol.cpp
+++ b/Engine/Source/Runtime/MovieSceneCapture/Private/CompositionGraphCaptureProtocol.cpp
@@ -14,8 +14,10 @@
 #include "BufferVisualizationData.h"
 #include "MovieSceneCaptureModule.h"
 #include "MovieSceneCaptureSettings.h"
+#include "Widgets/SWindow.h"
 #include "HDRHelper.h"
 
+
 struct FFrameCaptureViewExtension : public FSceneViewExtensionBase
 {
 	FFrameCaptureViewExtension( const FAutoRegister& AutoRegister, const TArray<FString>& InRenderPasses, bool bInCaptureFramesInHDR, int32 InHDRCompressionQuality, int32 InCaptureGamut, UMaterialInterface* InPostProcessingMaterial, bool bInDisableScreenPercentage)
diff --git a/Engine/Source/Runtime/MovieSceneCapture/Private/ImageSequenceProtocol.cpp b/Engine/Source/Runtime/MovieSceneCapture/Private/ImageSequenceProtocol.cpp
--- a/Engine/Source/Runtime/MovieSceneCapture/Private/ImageSequenceProtocol.cpp
+++ b/Engine/Source/Runtime/MovieSceneCapture/Private/ImageSequenceProtocol.cpp
@@ -16,6 +16,7 @@
 #include "MovieSceneCaptureModule.h"
 #include "MovieSceneCaptureSettings.h"
 #include "ImageWriteQueue.h"
+#include "Widgets/SWindow.h"
 #include "HDRHelper.h"
 
 struct FImageFrameData : IFramePayload
