[swin] Remove mention of experimental and audiomixer from the tooltip.

diff --git a/Engine/Source/Runtime/MovieSceneCapture/Public/MovieSceneCapture.h b/Engine/Source/Runtime/MovieSceneCapture/Public/MovieSceneCapture.h
--- a/Engine/Source/Runtime/MovieSceneCapture/Public/MovieSceneCapture.h
+++ b/Engine/Source/Runtime/MovieSceneCapture/Public/MovieSceneCapture.h
@@ -83,7 +83,7 @@ public:
 	UPROPERTY(config, EditAnywhere, NoClear, Category=CaptureSettings, DisplayName="Image Output Format", meta=(MetaClass="/Script/MovieSceneCapture.MovieSceneImageCaptureProtocolBase", HideViewOptions, ShowDisplayNames))
 	FSoftClassPath ImageCaptureProtocolType;
 	
-	/** The type of capture protocol to use for audio data. Requires experimental audio mixer (launch editor via with -audiomixer). */
+	/** The type of capture protocol to use for audio data. */
 	UPROPERTY(config, EditAnywhere, NoClear, Category=CaptureSettings, DisplayName="Audio Output Format", meta=(MetaClass="/Script/MovieSceneCapture.MovieSceneAudioCaptureProtocolBase", HideViewOptions, ShowDisplayNames))
 	FSoftClassPath AudioCaptureProtocolType;
 
diff --git a/Engine/Source/Runtime/MovieSceneCapture/Public/Protocols/AudioCaptureProtocol.h b/Engine/Source/Runtime/MovieSceneCapture/Public/Protocols/AudioCaptureProtocol.h
--- a/Engine/Source/Runtime/MovieSceneCapture/Public/Protocols/AudioCaptureProtocol.h
+++ b/Engine/Source/Runtime/MovieSceneCapture/Public/Protocols/AudioCaptureProtocol.h
@@ -22,11 +22,11 @@ class MOVIESCENECAPTURE_API UNullAudioCaptureProtocol : public UMovieSceneAudioC
 
 /**
 * This is an experimental audio capture implementation which captures the final output from the master submix.
-* This requires the new audiomixer (launch with "-audiomixer") and requires that your sequence can be played 
-* back in real-time (when rendering is disabled). If the sequence evaluation hitches the audio will become
-* desynchronized due to their being more time passed in real time (platform time) than in the sequence itself.
+* This requires that your sequence can be played back in real-time (when rendering is disabled). 
+* If the sequence evaluation hitches the audio will become desynchronized due to their being more time passed 
+* in real time (platform time) than in the sequence itself.
 */
-UCLASS(meta = (DisplayName = "Master Audio Submix (Experimental)", CommandLineID = "MasterAudioSubmix"))
+UCLASS(meta = (DisplayName = "Master Audio Submix", CommandLineID = "MasterAudioSubmix"))
 class MOVIESCENECAPTURE_API UMasterAudioSubmixCaptureProtocol : public UMovieSceneAudioCaptureProtocolBase
 {
 public:
