[Engine]

diff --git a/Engine/Source/Runtime/MovieScene/Public/MovieSceneSequenceTickManager.h b/Engine/Source/Runtime/MovieScene/Public/MovieSceneSequenceTickManager.h
--- a/Engine/Source/Runtime/MovieScene/Public/MovieSceneSequenceTickManager.h
+++ b/Engine/Source/Runtime/MovieScene/Public/MovieSceneSequenceTickManager.h
@@ -29,18 +29,18 @@ class IMovieSceneSequenceTickManagerClient;
  * Sequencer data is shared between all instances within the same group, allowing them to blend together.
  * Clients ticking at different intervals do not support blending with each other.
  */
-UCLASS()
-class MOVIESCENE_API UMovieSceneSequenceTickManager : public UObject
+UCLASS(MinimalAPI)
+class UMovieSceneSequenceTickManager : public UObject
 {
 public:
 	GENERATED_BODY()
 
-	UMovieSceneSequenceTickManager(const FObjectInitializer& Init);
+	MOVIESCENE_API UMovieSceneSequenceTickManager(const FObjectInitializer& Init);
 
 	/**
 	 * Retrieve the tick manager for the specified playback context
 	 */
-	static UMovieSceneSequenceTickManager* Get(UObject* PlaybackContext);
+	static MOVIESCENE_API UMovieSceneSequenceTickManager* Get(UObject* PlaybackContext);
 
 	/**
 	 * Register a new client to receive a tick at the specified tick interval.
@@ -49,26 +49,26 @@ public:
 	 * @param TickInterval    The interval at which to tick the client. 
 	 * @param InTickInterface The client that will receieve a tick
 	 */
-	void RegisterTickClient(const FMovieSceneSequenceTickInterval& TickInterval, TScriptInterface<IMovieSceneSequenceTickManagerClient> InTickInterface);
+	MOVIESCENE_API void RegisterTickClient(const FMovieSceneSequenceTickInterval& TickInterval, TScriptInterface<IMovieSceneSequenceTickManagerClient> InTickInterface);
 
 	/**
 	 * Unregister a previously registered client. The client must have been registered or an assertion will fail.
 	 *
 	 * @param InTickInterface The client to unregister
 	 */
-	void UnregisterTickClient(TScriptInterface<IMovieSceneSequenceTickManagerClient> InTickInterface);
+	MOVIESCENE_API void UnregisterTickClient(TScriptInterface<IMovieSceneSequenceTickManagerClient> InTickInterface);
 
 	/**
 	 * Retrieve the linker associated with the specified (resolved) tick interval.
 	 * @return The linker that owns all the data for the specified tick interval, or nullptr if there is none
 	 */
-	UMovieSceneEntitySystemLinker* GetLinker(const FMovieSceneSequenceTickInterval& TickInterval);
+	MOVIESCENE_API UMovieSceneEntitySystemLinker* GetLinker(const FMovieSceneSequenceTickInterval& TickInterval);
 
 	/**
 	 * Retrieve the runner associated with the specified (resolved) tick interval.
 	 * @return The runner that flushes evaluations for the specified tick interval, or nullptr if there is none
 	 */
-	TSharedPtr<FMovieSceneEntitySystemRunner> GetRunner(const FMovieSceneSequenceTickInterval& TickInterval);
+	MOVIESCENE_API TSharedPtr<FMovieSceneEntitySystemRunner> GetRunner(const FMovieSceneSequenceTickInterval& TickInterval);
 
 	/**
 	 * Add an action that will be executed once all the pending evaluations in this tick manager have been flushed.
@@ -76,36 +76,36 @@ public:
 	 *
 	 * @param Delegate A delegate to invoke when all evaluations have been flushed.
 	 */
-	void AddLatentAction(FMovieSceneSequenceLatentActionDelegate Delegate);
+	MOVIESCENE_API void AddLatentAction(FMovieSceneSequenceLatentActionDelegate Delegate);
 
 	/**
 	 * Execute any pending latent actions
 	 */
-	void RunLatentActions();
+	MOVIESCENE_API void RunLatentActions();
 
 	/**
 	 * Clear latent actions pertaining to the specified object
 	 */
-	void ClearLatentActions(UObject* Object);
+	MOVIESCENE_API void ClearLatentActions(UObject* Object);
 
 public:
 
 	/*~ -------------------------------------------------------------- */
 	/*~ Deprecated functions that previously used 'Actor' nomenclature */
 	UE_DEPRECATED(5.1, "Please use RegisterTickClient")
-	void RegisterSequenceActor(AActor* InActor);
+	MOVIESCENE_API void RegisterSequenceActor(AActor* InActor);
 	UE_DEPRECATED(5.1, "Please use RegisterTickClient")
-	void RegisterSequenceActor(AActor* InActor, TScriptInterface<IMovieSceneSequenceTickManagerClient> InActorInterface);
+	MOVIESCENE_API void RegisterSequenceActor(AActor* InActor, TScriptInterface<IMovieSceneSequenceTickManagerClient> InActorInterface);
 	UE_DEPRECATED(5.1, "Please use UnregisterTickClient")
-	void UnregisterSequenceActor(AActor* InActor);
+	MOVIESCENE_API void UnregisterSequenceActor(AActor* InActor);
 	UE_DEPRECATED(5.1, "Please use UnregisterTickClient")
-	void UnregisterSequenceActor(AActor* InActor, TScriptInterface<IMovieSceneSequenceTickManagerClient> InActorInterface);
+	MOVIESCENE_API void UnregisterSequenceActor(AActor* InActor, TScriptInterface<IMovieSceneSequenceTickManagerClient> InActorInterface);
 	/*~ -------------------------------------------------------------- */
 
 private:
 
 	/*~ UObject interface */
-	virtual void BeginDestroy() override;
+	MOVIESCENE_API virtual void BeginDestroy() override;
 
 private:
 
@@ -204,4 +204,4 @@ private:
 
 	/** Latent action manager for managing actions that need performing after evaluation */
 	FMovieSceneLatentActionManager LatentActionManager;
-};
\ No newline at end of file
+};
diff --git a/Engine/Source/Runtime/MovieScene/Public/MovieSceneSequenceTickManagerClient.h b/Engine/Source/Runtime/MovieScene/Public/MovieSceneSequenceTickManagerClient.h
--- a/Engine/Source/Runtime/MovieScene/Public/MovieSceneSequenceTickManagerClient.h
+++ b/Engine/Source/Runtime/MovieScene/Public/MovieSceneSequenceTickManagerClient.h
@@ -13,15 +13,15 @@ class FMovieSceneEntitySystemRunner;
 /**
  * Interface for objects that are to be ticked by the tick manager.
  */
-UINTERFACE()
-class MOVIESCENE_API UMovieSceneSequenceTickManagerClient
+UINTERFACE(MinimalAPI)
+class UMovieSceneSequenceTickManagerClient
 	: public UInterface
 {
 public:
 	GENERATED_BODY()
 };
 
-class MOVIESCENE_API IMovieSceneSequenceTickManagerClient
+class IMovieSceneSequenceTickManagerClient
 {
 public:
 	GENERATED_BODY()
diff --git a/Engine/Source/Runtime/MovieScene/Public/MovieSceneSequenceVisitor.h b/Engine/Source/Runtime/MovieScene/Public/MovieSceneSequenceVisitor.h
--- a/Engine/Source/Runtime/MovieScene/Public/MovieSceneSequenceVisitor.h
+++ b/Engine/Source/Runtime/MovieScene/Public/MovieSceneSequenceVisitor.h
@@ -63,9 +63,9 @@ namespace UE
 namespace MovieScene
 {
 
-struct MOVIESCENE_API FSubSequenceSpace
+struct FSubSequenceSpace
 {
-	FSubSequenceSpace();
+	MOVIESCENE_API FSubSequenceSpace();
 
 	/** Transform from the root time-space to the current sequence's time-space */
 	FMovieSceneSequenceTransform RootToSequenceTransform;
@@ -79,7 +79,7 @@ struct MOVIESCENE_API FSubSequenceSpace
 	int16 HierarchicalBias;
 };
 
-struct MOVIESCENE_API ISequenceVisitor
+struct ISequenceVisitor
 {
 	virtual ~ISequenceVisitor() {}
 
diff --git a/Engine/Source/Runtime/MovieScene/Public/MovieSceneSignedObject.h b/Engine/Source/Runtime/MovieScene/Public/MovieSceneSignedObject.h
--- a/Engine/Source/Runtime/MovieScene/Public/MovieSceneSignedObject.h
+++ b/Engine/Source/Runtime/MovieScene/Public/MovieSceneSignedObject.h
@@ -44,10 +44,10 @@ struct IDeferredSignedObjectChangeHandler
 	virtual void ResetImplicitScopedModifyDefer() = 0;
 };
 
-struct MOVIESCENE_API FScopedSignedObjectModifyDefer
+struct FScopedSignedObjectModifyDefer
 {
-	FScopedSignedObjectModifyDefer(bool bInForceFlush = false);
-	~FScopedSignedObjectModifyDefer();
+	MOVIESCENE_API FScopedSignedObjectModifyDefer(bool bInForceFlush = false);
+	MOVIESCENE_API ~FScopedSignedObjectModifyDefer();
 
 private:
 	bool bForceFlush;
@@ -59,24 +59,24 @@ private:
 
 
 
-UCLASS()
-class MOVIESCENE_API UMovieSceneSignedObject : public UObject
+UCLASS(MinimalAPI)
+class UMovieSceneSignedObject : public UObject
 {
 public:
 	GENERATED_BODY()
 
-	UMovieSceneSignedObject(const FObjectInitializer& Init);
+	MOVIESCENE_API UMovieSceneSignedObject(const FObjectInitializer& Init);
 
-	static TWeakPtr<UE::MovieScene::IDeferredSignedObjectChangeHandler> GetDeferredHandler();
-	static void SetDeferredHandler(TWeakPtr<UE::MovieScene::IDeferredSignedObjectChangeHandler>&& InHandler);
-	static void AddFlushSignal(TWeakPtr<UE::MovieScene::IDeferredSignedObjectFlushSignal> Signal);
-	static void ResetImplicitScopedModifyDefer();
+	static MOVIESCENE_API TWeakPtr<UE::MovieScene::IDeferredSignedObjectChangeHandler> GetDeferredHandler();
+	static MOVIESCENE_API void SetDeferredHandler(TWeakPtr<UE::MovieScene::IDeferredSignedObjectChangeHandler>&& InHandler);
+	static MOVIESCENE_API void AddFlushSignal(TWeakPtr<UE::MovieScene::IDeferredSignedObjectFlushSignal> Signal);
+	static MOVIESCENE_API void ResetImplicitScopedModifyDefer();
 
 	/**
 	 * Mark this object as having been changed in any way. Will regenerate this object's Signature
 	 * and schedule an event to be broadcast to notify subscribers of the change
 	 */
-	void MarkAsChanged();
+	MOVIESCENE_API void MarkAsChanged();
 
 	/**
 	 * Retrieve this object's signature that uniquely identifies its current state.
@@ -91,7 +91,7 @@ public:
 	/**
 	 * Immediately broadcast events for this object being changed
 	 */
-	void BroadcastChanged();
+	MOVIESCENE_API void BroadcastChanged();
 
 	/** Event that is triggered whenever this object's signature has changed */
 	DECLARE_EVENT(UMovieSceneSignedObject, FOnSignatureChanged)
@@ -102,14 +102,14 @@ public:
 
 public:
 
-	virtual void PostInitProperties() override;
-	virtual void PostLoad() override;
+	MOVIESCENE_API virtual void PostInitProperties() override;
+	MOVIESCENE_API virtual void PostLoad() override;
 
 #if WITH_EDITOR
-	virtual bool Modify(bool bAlwaysMarkDirty = true) override;
-	virtual void PostEditChangeProperty( struct FPropertyChangedEvent& PropertyChangedEvent) override;
-	virtual void PostEditUndo() override;
-	virtual void PostEditUndo(TSharedPtr<ITransactionObjectAnnotation> TransactionAnnotation) override;
+	MOVIESCENE_API virtual bool Modify(bool bAlwaysMarkDirty = true) override;
+	MOVIESCENE_API virtual void PostEditChangeProperty( struct FPropertyChangedEvent& PropertyChangedEvent) override;
+	MOVIESCENE_API virtual void PostEditUndo() override;
+	MOVIESCENE_API virtual void PostEditUndo(TSharedPtr<ITransactionObjectAnnotation> TransactionAnnotation) override;
 #endif
 
 private:
diff --git a/Engine/Source/Runtime/MovieScene/Public/MovieSceneTimeController.h b/Engine/Source/Runtime/MovieScene/Public/MovieSceneTimeController.h
--- a/Engine/Source/Runtime/MovieScene/Public/MovieSceneTimeController.h
+++ b/Engine/Source/Runtime/MovieScene/Public/MovieSceneTimeController.h
@@ -10,7 +10,7 @@
 
 class IMovieScenePlayer;
 
-struct MOVIESCENE_API FMovieSceneTimeController
+struct FMovieSceneTimeController
 {
 public:
 
@@ -19,12 +19,12 @@ public:
 	/**
 	 * Called whenever a sequence starts or resumes playback from a non-playing state
 	 */
-	void StartPlaying(const FQualifiedFrameTime& InStartTime);
+	MOVIESCENE_API void StartPlaying(const FQualifiedFrameTime& InStartTime);
 
 	/**
 	 * Called whenever a sequence stops playback
 	 */
-	void StopPlaying(const FQualifiedFrameTime& InStopTime);
+	MOVIESCENE_API void StopPlaying(const FQualifiedFrameTime& InStopTime);
 
 	/**
 	 * Ticks this controller
@@ -32,7 +32,7 @@ public:
 	 * @param DeltaSeconds     The tick delta in seconds, dilated by the current world settings global dilation
 	 * @param InPlayRate       The current play rate of the sequence
 	 */
-	void Tick(float DeltaSeconds, float InPlayRate);
+	MOVIESCENE_API void Tick(float DeltaSeconds, float InPlayRate);
 
 	/**
 	 * Request the current time based on the specified existing time and play rate.
@@ -41,17 +41,17 @@ public:
 	 * @param InCurrentTime    The current time of the sequence
 	 * @param InPlayRate       The current play rate of the sequence, multiplied by any world actor settings global dilation
 	 */
-	FFrameTime RequestCurrentTime(const FQualifiedFrameTime& InCurrentTime, float InPlayRate, FFrameRate InDisplayRate);
+	MOVIESCENE_API FFrameTime RequestCurrentTime(const FQualifiedFrameTime& InCurrentTime, float InPlayRate, FFrameRate InDisplayRate);
 
 	/**
 	 * Called when the status of the owning IMovieScenePlayer has changed
 	 */
-	void PlayerStatusChanged(EMovieScenePlayerStatus::Type InStatus, const FQualifiedFrameTime& InCurrentTime);
+	MOVIESCENE_API void PlayerStatusChanged(EMovieScenePlayerStatus::Type InStatus, const FQualifiedFrameTime& InCurrentTime);
 
 	/**
 	 * Called to stop and resume playback from the specified time
 	 */
-	void Reset(const FQualifiedFrameTime& InNewStartTime);
+	MOVIESCENE_API void Reset(const FQualifiedFrameTime& InNewStartTime);
 
 protected:
 
@@ -82,7 +82,7 @@ private:
 /**
  * A timing manager that retrieves its time from an external clock source
  */
-struct MOVIESCENE_API FMovieSceneTimeController_ExternalClock : FMovieSceneTimeController
+struct FMovieSceneTimeController_ExternalClock : FMovieSceneTimeController
 {
 protected:
 
@@ -94,9 +94,9 @@ protected:
 
 protected:
 
-	virtual void OnStartPlaying(const FQualifiedFrameTime& InStartTime) override;
-	virtual void OnStopPlaying(const FQualifiedFrameTime& InStopTime) override;
-	virtual FFrameTime OnRequestCurrentTime(const FQualifiedFrameTime& InCurrentTime, float InPlayRate) override;
+	MOVIESCENE_API virtual void OnStartPlaying(const FQualifiedFrameTime& InStartTime) override;
+	MOVIESCENE_API virtual void OnStopPlaying(const FQualifiedFrameTime& InStopTime) override;
+	MOVIESCENE_API virtual FFrameTime OnRequestCurrentTime(const FQualifiedFrameTime& InCurrentTime, float InPlayRate) override;
 
 private:
 
@@ -109,28 +109,28 @@ private:
 /**
  * A timing manager that retrieves its time from the platform clock
  */
-struct MOVIESCENE_API FMovieSceneTimeController_PlatformClock : FMovieSceneTimeController_ExternalClock
+struct FMovieSceneTimeController_PlatformClock : FMovieSceneTimeController_ExternalClock
 {
-	virtual double GetCurrentTime() const override;
+	MOVIESCENE_API virtual double GetCurrentTime() const override;
 };
 
 /**
  * A timing manager that retrieves its time from the audio clock
  */
-struct MOVIESCENE_API FMovieSceneTimeController_AudioClock : FMovieSceneTimeController_ExternalClock
+struct FMovieSceneTimeController_AudioClock : FMovieSceneTimeController_ExternalClock
 {
-	virtual double GetCurrentTime() const override;
+	MOVIESCENE_API virtual double GetCurrentTime() const override;
 };
 
 
 /**
 * A timing manager that retrieves its time relative to the Timecode clock
 */
-struct MOVIESCENE_API FMovieSceneTimeController_RelativeTimecodeClock : FMovieSceneTimeController_ExternalClock
+struct FMovieSceneTimeController_RelativeTimecodeClock : FMovieSceneTimeController_ExternalClock
 {
-	virtual FFrameTime OnRequestCurrentTime(const FQualifiedFrameTime& InCurrentTime, float InPlayRate) override;
-	virtual double GetCurrentTime() const override;
-	virtual void OnStopPlaying(const FQualifiedFrameTime& InStopTime) override;
+	MOVIESCENE_API virtual FFrameTime OnRequestCurrentTime(const FQualifiedFrameTime& InCurrentTime, float InPlayRate) override;
+	MOVIESCENE_API virtual double GetCurrentTime() const override;
+	MOVIESCENE_API virtual void OnStopPlaying(const FQualifiedFrameTime& InStopTime) override;
 
 protected:
 
@@ -144,16 +144,16 @@ protected:
 /**
 * A timing manager that retrieves its time from the Timecode clock
 */
-struct MOVIESCENE_API FMovieSceneTimeController_TimecodeClock : FMovieSceneTimeController
+struct FMovieSceneTimeController_TimecodeClock : FMovieSceneTimeController
 {
-	virtual FFrameTime OnRequestCurrentTime(const FQualifiedFrameTime& InCurrentTime, float InPlayRate) override;
+	MOVIESCENE_API virtual FFrameTime OnRequestCurrentTime(const FQualifiedFrameTime& InCurrentTime, float InPlayRate) override;
 };
 
 
 /**
  * A timing manager that accumulates delta times from a world
  */
-struct MOVIESCENE_API FMovieSceneTimeController_Tick : FMovieSceneTimeController
+struct FMovieSceneTimeController_Tick : FMovieSceneTimeController
 {
 	FMovieSceneTimeController_Tick()
 		: CurrentOffsetSeconds(0.0)
@@ -161,9 +161,9 @@ struct MOVIESCENE_API FMovieSceneTimeController_Tick : FMovieSceneTimeController
 
 protected:
 
-	virtual void OnTick(float DeltaSeconds, float InPlayRate) override;
-	virtual void OnStartPlaying(const FQualifiedFrameTime& InStartTime) override;
-	virtual FFrameTime OnRequestCurrentTime(const FQualifiedFrameTime& InCurrentTime, float InPlayRate) override;
+	MOVIESCENE_API virtual void OnTick(float DeltaSeconds, float InPlayRate) override;
+	MOVIESCENE_API virtual void OnStartPlaying(const FQualifiedFrameTime& InStartTime) override;
+	MOVIESCENE_API virtual FFrameTime OnRequestCurrentTime(const FQualifiedFrameTime& InCurrentTime, float InPlayRate) override;
 
 private:
 	double CurrentOffsetSeconds;
@@ -172,7 +172,7 @@ private:
 /**
  * A timing manager that plays every display frame for a certain number of seconds
  */
-struct MOVIESCENE_API FMovieSceneTimeController_PlayEveryFrame : FMovieSceneTimeController
+struct FMovieSceneTimeController_PlayEveryFrame : FMovieSceneTimeController
 {
 	FMovieSceneTimeController_PlayEveryFrame()
 		: PreviousPlatformTime(0.f)
@@ -180,8 +180,8 @@ struct MOVIESCENE_API FMovieSceneTimeController_PlayEveryFrame : FMovieSceneTime
 
 protected:
 
-	virtual void OnStartPlaying(const FQualifiedFrameTime& InStartTime) override;
-	virtual FFrameTime OnRequestCurrentTime(const FQualifiedFrameTime& InCurrentTime, float InPlayRate) override;
+	MOVIESCENE_API virtual void OnStartPlaying(const FQualifiedFrameTime& InStartTime) override;
+	MOVIESCENE_API virtual FFrameTime OnRequestCurrentTime(const FQualifiedFrameTime& InCurrentTime, float InPlayRate) override;
 
 private:
 	double PreviousPlatformTime;
diff --git a/Engine/Source/Runtime/MovieScene/Public/MovieSceneTimeHelpers.h b/Engine/Source/Runtime/MovieScene/Public/MovieSceneTimeHelpers.h
--- a/Engine/Source/Runtime/MovieScene/Public/MovieSceneTimeHelpers.h
+++ b/Engine/Source/Runtime/MovieScene/Public/MovieSceneTimeHelpers.h
@@ -20,7 +20,7 @@ namespace UE
 namespace MovieScene
 {
 
-class MOVIESCENE_API TimeHelpers
+class TimeHelpers
 {
 public:
 
@@ -28,7 +28,7 @@ public:
  * Migrate the frame times of the movie scene from the source frame rate to the destination frame rate
  */
 
-static void MigrateFrameTimes(FFrameRate SourceRate, FFrameRate DestinationRate, UMovieScene* MovieScene, bool bApplyRecursively = false);
+static MOVIESCENE_API void MigrateFrameTimes(FFrameRate SourceRate, FFrameRate DestinationRate, UMovieScene* MovieScene, bool bApplyRecursively = false);
 
 };
 
diff --git a/Engine/Source/Runtime/MovieScene/Public/MovieSceneTrack.h b/Engine/Source/Runtime/MovieScene/Public/MovieSceneTrack.h
--- a/Engine/Source/Runtime/MovieScene/Public/MovieSceneTrack.h
+++ b/Engine/Source/Runtime/MovieScene/Public/MovieSceneTrack.h
@@ -170,7 +170,7 @@ class UMovieSceneTrack
 public:
 
 	MOVIESCENE_API UMovieSceneTrack(const FObjectInitializer& InInitializer);
-	MOVIESCENE_API ~UMovieSceneTrack() {};
+	~UMovieSceneTrack() {};
 
 public:
 
@@ -346,18 +346,18 @@ public:
 	MOVIESCENE_API bool FixRowIndices();
 
 	/** Called when row indices have been fixed up */
-	MOVIESCENE_API virtual void OnRowIndicesChanged(const TMap<int32, int32>& NewToOldRowIndices) {}
+	virtual void OnRowIndicesChanged(const TMap<int32, int32>& NewToOldRowIndices) {}
 
 	/**
 	* @return Whether evaluation of this track should be disabled due to mute/solo settings
 	*/
-	MOVIESCENE_API bool IsEvalDisabled() const { return bIsEvalDisabled; }
+	bool IsEvalDisabled() const { return bIsEvalDisabled; }
 	MOVIESCENE_API bool IsRowEvalDisabled(int32 RowIndex) const;
 
 	/**
 	* Called by Sequencer to set whether evaluation of this track should be disabled due to mute/solo settings
 	*/
-	MOVIESCENE_API void SetEvalDisabled(bool bEvalDisabled) { bIsEvalDisabled = bEvalDisabled; }
+	void SetEvalDisabled(bool bEvalDisabled) { bIsEvalDisabled = bEvalDisabled; }
 	MOVIESCENE_API void SetRowEvalDisabled(bool bEvalDisabled, int32 RowIndex);
 
 public:
diff --git a/Engine/Source/Runtime/MovieScene/Public/Sections/MovieSceneHookSection.h b/Engine/Source/Runtime/MovieScene/Public/Sections/MovieSceneHookSection.h
--- a/Engine/Source/Runtime/MovieScene/Public/Sections/MovieSceneHookSection.h
+++ b/Engine/Source/Runtime/MovieScene/Public/Sections/MovieSceneHookSection.h
@@ -23,8 +23,8 @@ template <typename ElementType> class TRange;
 /**
  * 
  */
-UCLASS()
-class MOVIESCENE_API UMovieSceneHookSection
+UCLASS(MinimalAPI)
+class UMovieSceneHookSection
 	: public UMovieSceneSection
 	, public IMovieSceneEntityProvider
 	, public IMovieSceneEvaluationHook
@@ -33,7 +33,7 @@ public:
 
 	GENERATED_BODY()
 
-	UMovieSceneHookSection(const FObjectInitializer&);
+	MOVIESCENE_API UMovieSceneHookSection(const FObjectInitializer&);
 
 	virtual TArrayView<const FFrameNumber> GetTriggerTimes() const { return TArrayView<const FFrameNumber>(); }
 
@@ -49,11 +49,11 @@ protected:
 
 	*/
 
-	virtual void ImportEntityImpl(UMovieSceneEntitySystemLinker* EntityLinker, const FEntityImportParams& Params, FImportedEntity* OutImportedEntity) override;
-	virtual bool PopulateEvaluationFieldImpl(const TRange<FFrameNumber>& EffectiveRange, const FMovieSceneEvaluationFieldEntityMetaData& InMetaData, FMovieSceneEntityComponentFieldBuilder* OutFieldBuilder) override;
+	MOVIESCENE_API virtual void ImportEntityImpl(UMovieSceneEntitySystemLinker* EntityLinker, const FEntityImportParams& Params, FImportedEntity* OutImportedEntity) override;
+	MOVIESCENE_API virtual bool PopulateEvaluationFieldImpl(const TRange<FFrameNumber>& EffectiveRange, const FMovieSceneEvaluationFieldEntityMetaData& InMetaData, FMovieSceneEntityComponentFieldBuilder* OutFieldBuilder) override;
 
-	void ImportRangedEntity(UMovieSceneEntitySystemLinker* EntityLinker, const FEntityImportParams& Params, FImportedEntity* OutImportedEntity);
-	void ImportTriggerEntity(UMovieSceneEntitySystemLinker* EntityLinker, const FEntityImportParams& Params, FImportedEntity* OutImportedEntity);
+	MOVIESCENE_API void ImportRangedEntity(UMovieSceneEntitySystemLinker* EntityLinker, const FEntityImportParams& Params, FImportedEntity* OutImportedEntity);
+	MOVIESCENE_API void ImportTriggerEntity(UMovieSceneEntitySystemLinker* EntityLinker, const FEntityImportParams& Params, FImportedEntity* OutImportedEntity);
 
 protected:
 
diff --git a/Engine/Source/Runtime/MovieScene/Public/Sections/MovieSceneSubSection.h b/Engine/Source/Runtime/MovieScene/Public/Sections/MovieSceneSubSection.h
--- a/Engine/Source/Runtime/MovieScene/Public/Sections/MovieSceneSubSection.h
+++ b/Engine/Source/Runtime/MovieScene/Public/Sections/MovieSceneSubSection.h
@@ -59,8 +59,8 @@ struct FSubSequenceInstanceDataParams
 /**
  * Implements a section in sub-sequence tracks.
  */
-UCLASS(BlueprintType, config = EditorPerProjectUserSettings)
-class MOVIESCENE_API UMovieSceneSubSection
+UCLASS(BlueprintType, config = EditorPerProjectUserSettings, MinimalAPI)
+class UMovieSceneSubSection
 	: public UMovieSceneSection
 	, public IMovieSceneEntityProvider
 {
@@ -69,7 +69,7 @@ class MOVIESCENE_API UMovieSceneSubSection
 public:
 
 	/** Object constructor. */
-	UMovieSceneSubSection(const FObjectInitializer& ObjInitializer);
+	MOVIESCENE_API UMovieSceneSubSection(const FObjectInitializer& ObjInitializer);
 
 	/**
 	 * Get the sequence that is assigned to this section.
@@ -78,39 +78,39 @@ public:
 	 * @see SetSequence
 	 */
 	UFUNCTION(BlueprintPure, Category = "Sequencer|Section")
-	UMovieSceneSequence* GetSequence() const;
+	MOVIESCENE_API UMovieSceneSequence* GetSequence() const;
 
 	/**
 	 * Get the path name to this sub section from the outer moviescene
 	 */
-	FString GetPathNameInMovieScene() const;
+	MOVIESCENE_API FString GetPathNameInMovieScene() const;
 
 	/**
 	 * Get this sub section's sequence ID
 	 */
-	FMovieSceneSequenceID GetSequenceID() const;
+	MOVIESCENE_API FMovieSceneSequenceID GetSequenceID() const;
 
 	/** Generate subsequence data */
-	virtual FMovieSceneSubSequenceData GenerateSubSequenceData(const FSubSequenceInstanceDataParams& Params) const;
+	MOVIESCENE_API virtual FMovieSceneSubSequenceData GenerateSubSequenceData(const FSubSequenceInstanceDataParams& Params) const;
 
 public:
 
 	/**
 	 * Gets the transform that converts time from this section's time-base to its inner sequence's
 	 */
-	FMovieSceneSequenceTransform OuterToInnerTransform() const;
+	MOVIESCENE_API FMovieSceneSequenceTransform OuterToInnerTransform() const;
 
 	/**
 	 * Gets the playrange of the inner sequence, in the inner sequence's time space, trimmed with any start/end offsets,
 	 * and validated to make sure we get at least a 1-frame long playback range (e.g. in the case where excessive
 	 * trimming results in an invalid range).
 	 */
-	bool GetValidatedInnerPlaybackRange(TRange<FFrameNumber>& OutInnerPlaybackRange) const;
+	MOVIESCENE_API bool GetValidatedInnerPlaybackRange(TRange<FFrameNumber>& OutInnerPlaybackRange) const;
 
 	/**
 	 * Helper function used by the above method, but accessible for other uses like track editors.
 	 */
-	static TRange<FFrameNumber> GetValidatedInnerPlaybackRange(const FMovieSceneSectionParameters& SubSectionParameters, const UMovieScene& InnerMovieScene);
+	static MOVIESCENE_API TRange<FFrameNumber> GetValidatedInnerPlaybackRange(const FMovieSceneSectionParameters& SubSectionParameters, const UMovieScene& InnerMovieScene);
 
 	/**
 	 * Sets the sequence played by this section.
@@ -119,19 +119,19 @@ public:
 	 * @see GetSequence
 	 */
 	UFUNCTION(BlueprintCallable, Category = "Sequencer|Section")
-	void SetSequence(UMovieSceneSequence* Sequence);
+	MOVIESCENE_API void SetSequence(UMovieSceneSequence* Sequence);
 
-	virtual void PostLoad() override;
+	MOVIESCENE_API virtual void PostLoad() override;
 
 #if WITH_EDITOR
-	virtual void PreEditChange(FProperty* PropertyAboutToChange) override;
-	virtual void PostEditChangeProperty(FPropertyChangedEvent& PropertyChangedEvent) override;
+	MOVIESCENE_API virtual void PreEditChange(FProperty* PropertyAboutToChange) override;
+	MOVIESCENE_API virtual void PostEditChangeProperty(FPropertyChangedEvent& PropertyChangedEvent) override;
 
 	/** Delegate to fire when our sequence is changed in the property editor */
 	FOnSequenceChanged& OnSequenceChanged() { return OnSequenceChangedDelegate; }
 #endif
 
-	FFrameNumber MapTimeToSectionFrame(FFrameTime InPosition) const;
+	MOVIESCENE_API FFrameNumber MapTimeToSectionFrame(FFrameTime InPosition) const;
 
 	EMovieSceneServerClientMask GetNetworkMask() const
 	{
@@ -146,19 +146,19 @@ public:
 public:
 
 	//~ UMovieSceneSection interface
-	virtual TOptional<TRange<FFrameNumber> > GetAutoSizeRange() const override;
-	virtual UMovieSceneSection* SplitSection( FQualifiedFrameTime SplitTime, bool bDeleteKeys ) override;
-	virtual void TrimSection( FQualifiedFrameTime TrimTime, bool bTrimLeft, bool bDeleteKeys) override;
+	MOVIESCENE_API virtual TOptional<TRange<FFrameNumber> > GetAutoSizeRange() const override;
+	MOVIESCENE_API virtual UMovieSceneSection* SplitSection( FQualifiedFrameTime SplitTime, bool bDeleteKeys ) override;
+	MOVIESCENE_API virtual void TrimSection( FQualifiedFrameTime TrimTime, bool bTrimLeft, bool bDeleteKeys) override;
 	virtual TOptional<FFrameTime> GetOffsetTime() const override { return TOptional<FFrameTime>(FFrameTime(Parameters.StartFrameOffset)); }
-	virtual void GetSnapTimes(TArray<FFrameNumber>& OutSnapTimes, bool bGetSectionBorders) const override;
-	virtual void MigrateFrameTimes(FFrameRate SourceRate, FFrameRate DestinationRate) override;
+	MOVIESCENE_API virtual void GetSnapTimes(TArray<FFrameNumber>& OutSnapTimes, bool bGetSectionBorders) const override;
+	MOVIESCENE_API virtual void MigrateFrameTimes(FFrameRate SourceRate, FFrameRate DestinationRate) override;
 
 protected:
 
-	void BuildDefaultSubSectionComponents(UMovieSceneEntitySystemLinker* EntityLinker, const UE::MovieScene::FEntityImportParams& Params, UE::MovieScene::FImportedEntity* OutImportedEntity) const;
+	MOVIESCENE_API void BuildDefaultSubSectionComponents(UMovieSceneEntitySystemLinker* EntityLinker, const UE::MovieScene::FEntityImportParams& Params, UE::MovieScene::FImportedEntity* OutImportedEntity) const;
 
-	virtual bool PopulateEvaluationFieldImpl(const TRange<FFrameNumber>& EffectiveRange, const FMovieSceneEvaluationFieldEntityMetaData& InMetaData, FMovieSceneEntityComponentFieldBuilder* OutFieldBuilder) override;
-	virtual void ImportEntityImpl(UMovieSceneEntitySystemLinker* EntityLinker, const FEntityImportParams& Params, FImportedEntity* OutImportedEntity) override;
+	MOVIESCENE_API virtual bool PopulateEvaluationFieldImpl(const TRange<FFrameNumber>& EffectiveRange, const FMovieSceneEvaluationFieldEntityMetaData& InMetaData, FMovieSceneEntityComponentFieldBuilder* OutFieldBuilder) override;
+	MOVIESCENE_API virtual void ImportEntityImpl(UMovieSceneEntitySystemLinker* EntityLinker, const FEntityImportParams& Params, FImportedEntity* OutImportedEntity) override;
 
 public:
 
diff --git a/Engine/Source/Runtime/MovieScene/Public/Tracks/MovieSceneCachedTrack.h b/Engine/Source/Runtime/MovieScene/Public/Tracks/MovieSceneCachedTrack.h
--- a/Engine/Source/Runtime/MovieScene/Public/Tracks/MovieSceneCachedTrack.h
+++ b/Engine/Source/Runtime/MovieScene/Public/Tracks/MovieSceneCachedTrack.h
@@ -6,8 +6,8 @@
 #include "UObject/Interface.h"
 #include "MovieSceneCachedTrack.generated.h"
 
-UINTERFACE()
-class MOVIESCENE_API UMovieSceneCachedTrack : public UInterface
+UINTERFACE(MinimalAPI)
+class UMovieSceneCachedTrack : public UInterface
 {
 	GENERATED_BODY()
 };
@@ -15,7 +15,7 @@ class MOVIESCENE_API UMovieSceneCachedTrack : public UInterface
 /**
  * Can be implemented by tracks that hold cached data. Used by the take recorder plugin to regenerate the cache when recording. 
  */
-class MOVIESCENE_API IMovieSceneCachedTrack
+class IMovieSceneCachedTrack
 {
 	GENERATED_BODY()
 
diff --git a/Engine/Source/Runtime/MovieScene/Public/Tracks/MovieSceneSpawnTrack.h b/Engine/Source/Runtime/MovieScene/Public/Tracks/MovieSceneSpawnTrack.h
--- a/Engine/Source/Runtime/MovieScene/Public/Tracks/MovieSceneSpawnTrack.h
+++ b/Engine/Source/Runtime/MovieScene/Public/Tracks/MovieSceneSpawnTrack.h
@@ -26,13 +26,13 @@ template <typename ElementType> class TRange;
 /**
  * Handles when a spawnable should be spawned and destroyed
  */
-UCLASS()
-class MOVIESCENE_API UMovieSceneSpawnTrack
+UCLASS(MinimalAPI)
+class UMovieSceneSpawnTrack
 	: public UMovieSceneTrack
 {
 public:
 
-	UMovieSceneSpawnTrack(const FObjectInitializer& Obj);
+	MOVIESCENE_API UMovieSceneSpawnTrack(const FObjectInitializer& Obj);
 	GENERATED_BODY()
 
 public:
@@ -51,30 +51,30 @@ public:
 
 	static uint16 GetEvaluationPriority() { return uint16(0xFFF); }
 
-	void PopulateSpawnedRangeMask(const TRange<FFrameNumber>& InOverlap, TArray<TRange<FFrameNumber>, TInlineAllocator<1>>& OutRanges) const;
+	MOVIESCENE_API void PopulateSpawnedRangeMask(const TRange<FFrameNumber>& InOverlap, TArray<TRange<FFrameNumber>, TInlineAllocator<1>>& OutRanges) const;
 
 public:
 
 	// UMovieSceneTrack interface
-	virtual bool SupportsType(TSubclassOf<UMovieSceneSection> SectionClass) const override;
-	virtual UMovieSceneSection* CreateNewSection() override;
-	virtual bool HasSection(const UMovieSceneSection& Section) const override;
-	virtual void AddSection(UMovieSceneSection& Section) override;
-	virtual void RemoveSection(UMovieSceneSection& Section) override;
-	virtual void RemoveSectionAt(int32 SectionIndex) override;
-	virtual void RemoveAllAnimationData() override;
-	virtual bool IsEmpty() const override;
-	virtual const TArray<UMovieSceneSection*>& GetAllSections() const override;
+	MOVIESCENE_API virtual bool SupportsType(TSubclassOf<UMovieSceneSection> SectionClass) const override;
+	MOVIESCENE_API virtual UMovieSceneSection* CreateNewSection() override;
+	MOVIESCENE_API virtual bool HasSection(const UMovieSceneSection& Section) const override;
+	MOVIESCENE_API virtual void AddSection(UMovieSceneSection& Section) override;
+	MOVIESCENE_API virtual void RemoveSection(UMovieSceneSection& Section) override;
+	MOVIESCENE_API virtual void RemoveSectionAt(int32 SectionIndex) override;
+	MOVIESCENE_API virtual void RemoveAllAnimationData() override;
+	MOVIESCENE_API virtual bool IsEmpty() const override;
+	MOVIESCENE_API virtual const TArray<UMovieSceneSection*>& GetAllSections() const override;
 
 	//~ UObject interface
-	virtual void PostLoad() override;
+	MOVIESCENE_API virtual void PostLoad() override;
 
 #if WITH_EDITOR
-	virtual ECookOptimizationFlags GetCookOptimizationFlags() const override;
+	MOVIESCENE_API virtual ECookOptimizationFlags GetCookOptimizationFlags() const override;
 #endif
 
 #if WITH_EDITORONLY_DATA
-	virtual FText GetDisplayName() const override;
+	MOVIESCENE_API virtual FText GetDisplayName() const override;
 #endif
 
 protected:
diff --git a/Engine/Source/Runtime/MovieScene/Public/Tracks/MovieSceneSubTrack.h b/Engine/Source/Runtime/MovieScene/Public/Tracks/MovieSceneSubTrack.h
--- a/Engine/Source/Runtime/MovieScene/Public/Tracks/MovieSceneSubTrack.h
+++ b/Engine/Source/Runtime/MovieScene/Public/Tracks/MovieSceneSubTrack.h
@@ -25,15 +25,15 @@ struct FMovieSceneSegmentCompilerRules;
 /**
  * A track that holds sub-sequences within a larger sequence.
  */
-UCLASS()
-class MOVIESCENE_API UMovieSceneSubTrack
+UCLASS(MinimalAPI)
+class UMovieSceneSubTrack
 	: public UMovieSceneNameableTrack
 {
 	GENERATED_BODY()
 
 public:
 
-	UMovieSceneSubTrack( const FObjectInitializer& ObjectInitializer );
+	MOVIESCENE_API UMovieSceneSubTrack( const FObjectInitializer& ObjectInitializer );
 
 	/**
 	 * Adds a movie scene section at the requested time.
@@ -54,7 +54,7 @@ public:
 	 * @param bInsertSequence Whether or not to insert the sequence and push existing sequences out
 	 * @return The newly created sub section
 	 */
-	virtual UMovieSceneSubSection* AddSequenceOnRow(UMovieSceneSequence* Sequence, FFrameNumber StartTime, int32 Duration, int32 RowIndex);
+	MOVIESCENE_API virtual UMovieSceneSubSection* AddSequenceOnRow(UMovieSceneSequence* Sequence, FFrameNumber StartTime, int32 Duration, int32 RowIndex);
 
 	/**
 	 * Check whether this track contains the given sequence.
@@ -64,25 +64,25 @@ public:
 	 * @param SectionToSkip Skip this section when searching the track (ie. the section is already set to this sequence). 
 	 * @return true if the sequence is in this track, false otherwise.
 	 */
-	bool ContainsSequence(const UMovieSceneSequence& Sequence, bool Recursively = false, const UMovieSceneSection* SectionToSkip = nullptr) const;
+	MOVIESCENE_API bool ContainsSequence(const UMovieSceneSequence& Sequence, bool Recursively = false, const UMovieSceneSection* SectionToSkip = nullptr) const;
 
 public:
 
 	// UMovieSceneTrack interface
 
-	virtual void AddSection(UMovieSceneSection& Section) override;
-	virtual bool SupportsType(TSubclassOf<UMovieSceneSection> SectionClass) const override;
-	virtual UMovieSceneSection* CreateNewSection() override;
-	virtual const TArray<UMovieSceneSection*>& GetAllSections() const override;
-	virtual bool HasSection(const UMovieSceneSection& Section) const override;
-	virtual bool IsEmpty() const override;
-	virtual void RemoveAllAnimationData() override;
-	virtual void RemoveSection(UMovieSceneSection& Section) override;
-	virtual void RemoveSectionAt(int32 SectionIndex) override;
-	virtual bool SupportsMultipleRows() const override;
+	MOVIESCENE_API virtual void AddSection(UMovieSceneSection& Section) override;
+	MOVIESCENE_API virtual bool SupportsType(TSubclassOf<UMovieSceneSection> SectionClass) const override;
+	MOVIESCENE_API virtual UMovieSceneSection* CreateNewSection() override;
+	MOVIESCENE_API virtual const TArray<UMovieSceneSection*>& GetAllSections() const override;
+	MOVIESCENE_API virtual bool HasSection(const UMovieSceneSection& Section) const override;
+	MOVIESCENE_API virtual bool IsEmpty() const override;
+	MOVIESCENE_API virtual void RemoveAllAnimationData() override;
+	MOVIESCENE_API virtual void RemoveSection(UMovieSceneSection& Section) override;
+	MOVIESCENE_API virtual void RemoveSectionAt(int32 SectionIndex) override;
+	MOVIESCENE_API virtual bool SupportsMultipleRows() const override;
 
 #if WITH_EDITORONLY_DATA
-	virtual FText GetDefaultDisplayName() const override;
+	MOVIESCENE_API virtual FText GetDefaultDisplayName() const override;
 #endif
 
 protected:
diff --git a/Engine/Source/Runtime/MovieScene/Public/TransformData.h b/Engine/Source/Runtime/MovieScene/Public/TransformData.h
--- a/Engine/Source/Runtime/MovieScene/Public/TransformData.h
+++ b/Engine/Source/Runtime/MovieScene/Public/TransformData.h
@@ -8,7 +8,7 @@
 /**
 * Stores information about a transform for the purpose of adding keys to a transform section
 */
-struct MOVIESCENE_API FTransformData
+struct FTransformData
 {
 	/** Translation component */
 	FVector Translation;
diff --git a/Engine/Source/Runtime/MovieSceneCapture/Public/FrameGrabber.h b/Engine/Source/Runtime/MovieSceneCapture/Public/FrameGrabber.h
--- a/Engine/Source/Runtime/MovieSceneCapture/Public/FrameGrabber.h
+++ b/Engine/Source/Runtime/MovieSceneCapture/Public/FrameGrabber.h
@@ -10,24 +10,24 @@
 class SWindow;
 
 /** A single, managed surface used as a render target resolution destination  */
-struct MOVIESCENECAPTURE_API FViewportSurfaceReader
+struct FViewportSurfaceReader
 {
 	/** Constructor */
-	FViewportSurfaceReader(EPixelFormat InPixelFormat, FIntPoint InBufferSize);
+	MOVIESCENECAPTURE_API FViewportSurfaceReader(EPixelFormat InPixelFormat, FIntPoint InBufferSize);
 
 	/** Destructor */
-	~FViewportSurfaceReader();
+	MOVIESCENECAPTURE_API ~FViewportSurfaceReader();
 
 	/** Initialize this reader so that it can be waited on. */
-	void Initialize();
+	MOVIESCENECAPTURE_API void Initialize();
 
 	/** Wait for this reader to become available, if it's currently in use */
-	void BlockUntilAvailable();
+	MOVIESCENECAPTURE_API void BlockUntilAvailable();
 
 	/** Safely resets the state of the wait event. When doing latent surface reading sometimes we may want to just bail on reading a given frame.
 	  * Should only be performed after flushing rendering commands.
 	  */
-	void Reset();
+	MOVIESCENECAPTURE_API void Reset();
 
 	/**
 	 * Resolve the specified viewport RHI, calling the specified callback with the result.
@@ -35,10 +35,10 @@ struct MOVIESCENECAPTURE_API FViewportSurfaceReader
 	 * @param	BackBuffer		The backbuffer to resolve
 	 * @param	Callback 		Callback to call with the locked texture data. This will be called on an undefined thread.
 	 */
-	void ResolveRenderTarget(FViewportSurfaceReader* RenderToReadback, const FTexture2DRHIRef& BackBuffer, TFunction<void(FColor*, int32, int32)> Callback);
+	MOVIESCENECAPTURE_API void ResolveRenderTarget(FViewportSurfaceReader* RenderToReadback, const FTexture2DRHIRef& BackBuffer, TFunction<void(FColor*, int32, int32)> Callback);
 
 	/** Get the current size of the texture */
-	FIntPoint GetCurrentSize() const;
+	MOVIESCENECAPTURE_API FIntPoint GetCurrentSize() const;
 
 	/** Set the rectangle within which to read pixels */
 	void SetCaptureRect(FIntRect InCaptureRect) { CaptureRect = InCaptureRect; }
@@ -51,7 +51,7 @@ struct MOVIESCENECAPTURE_API FViewportSurfaceReader
 protected:
 
 	/** Set up this surface to the specified width/height */
-	void Resize(uint32 Width, uint32 Height);
+	MOVIESCENECAPTURE_API void Resize(uint32 Width, uint32 Height);
 
 	/** Whether this surface reader is enabled or not */
 	FThreadSafeBool bEnabled;
@@ -121,7 +121,7 @@ private:
  * to resolve the viewport render target into a specific index into this array. This means we can
  * resolve the render target data without having to wait, or flush rendering commands.
  */
-class MOVIESCENECAPTURE_API FFrameGrabber
+class FFrameGrabber
 {
 public:
 	/**
@@ -132,43 +132,43 @@ public:
 	 * @param InPixelFormat			The desired pixel format to store captured frames as
 	 * @param InNumSurfaces			The number of destination surfaces contained in our buffer 
 	 */
-	FFrameGrabber(TSharedRef<FSceneViewport> Viewport, FIntPoint DesiredBufferSize, EPixelFormat InPixelFormat = PF_B8G8R8A8, uint32 NumSurfaces = 3);
+	MOVIESCENECAPTURE_API FFrameGrabber(TSharedRef<FSceneViewport> Viewport, FIntPoint DesiredBufferSize, EPixelFormat InPixelFormat = PF_B8G8R8A8, uint32 NumSurfaces = 3);
 
 	/** Destructor */
-	~FFrameGrabber();
+	MOVIESCENECAPTURE_API ~FFrameGrabber();
 
 public:
 
 	/** Instruct the frame grabber to start capturing frames */
-	void StartCapturingFrames();
+	MOVIESCENECAPTURE_API void StartCapturingFrames();
 
 	/** Check whether we're capturing frames or not */
-	bool IsCapturingFrames() const;
+	MOVIESCENECAPTURE_API bool IsCapturingFrames() const;
 
 	/** Instruct the frame grabber capture this frame, when it receives an event from slate */
-	void CaptureThisFrame(FFramePayloadPtr Payload);
+	MOVIESCENECAPTURE_API void CaptureThisFrame(FFramePayloadPtr Payload);
 
 	/** Stop capturing frames */
-	void StopCapturingFrames();
+	MOVIESCENECAPTURE_API void StopCapturingFrames();
 
 	/** Shut down this grabber, ensuring that any threaded operations are finished */
-	void Shutdown();
+	MOVIESCENECAPTURE_API void Shutdown();
 
 public:
 
 	/** Check whether we have any outstanding frames or not */
-	bool HasOutstandingFrames() const;
+	MOVIESCENECAPTURE_API bool HasOutstandingFrames() const;
 
 	/** Retrieve any frames we may have captured */
-	TArray<FCapturedFrameData> GetCapturedFrames();
+	MOVIESCENECAPTURE_API TArray<FCapturedFrameData> GetCapturedFrames();
 
 protected:
 	
 	/** Callback for when a backbuffer is ready for reading (called on render thread) */
-	void OnBackBufferReadyToPresentCallback(SWindow& SlateWindow, const FTexture2DRHIRef& BackBuffer);
+	MOVIESCENECAPTURE_API void OnBackBufferReadyToPresentCallback(SWindow& SlateWindow, const FTexture2DRHIRef& BackBuffer);
 
 	/** Called when the specified surface index has been locked for reading with the render target data (called on render thread)  */
-	void OnFrameReady(int32 SurfaceIndex, FColor* ColorBuffer, int32 Width, int32 Height);
+	MOVIESCENECAPTURE_API void OnFrameReady(int32 SurfaceIndex, FColor* ColorBuffer, int32 Width, int32 Height);
 
 private:
 
diff --git a/Engine/Source/Runtime/MovieSceneCapture/Public/LevelCapture.h b/Engine/Source/Runtime/MovieSceneCapture/Public/LevelCapture.h
--- a/Engine/Source/Runtime/MovieSceneCapture/Public/LevelCapture.h
+++ b/Engine/Source/Runtime/MovieSceneCapture/Public/LevelCapture.h
@@ -11,11 +11,11 @@
 class AActor;
 class FSceneViewport;
 
-UCLASS()
-class MOVIESCENECAPTURE_API ULevelCapture : public UMovieSceneCapture
+UCLASS(MinimalAPI)
+class ULevelCapture : public UMovieSceneCapture
 {
 public:
-	ULevelCapture(const FObjectInitializer&);
+	MOVIESCENECAPTURE_API ULevelCapture(const FObjectInitializer&);
 
 	GENERATED_BODY()
 
@@ -24,10 +24,10 @@ public:
 	bool bAutoStartCapture;
 
 	/** Specify a prerequisite actor that must be set up before we start capturing */
-	void SetPrerequisiteActor(AActor* Prereq);
+	MOVIESCENECAPTURE_API void SetPrerequisiteActor(AActor* Prereq);
 
-	virtual void Initialize(TSharedPtr<FSceneViewport> InViewport, int32 PIEInstance = -1) override;
-	virtual void OnTick(float DeltaSeconds) override;
+	MOVIESCENECAPTURE_API virtual void Initialize(TSharedPtr<FSceneViewport> InViewport, int32 PIEInstance = -1) override;
+	MOVIESCENECAPTURE_API virtual void OnTick(float DeltaSeconds) override;
 	
 private:
 
diff --git a/Engine/Source/Runtime/MovieSceneCapture/Public/MovieSceneCapture.h b/Engine/Source/Runtime/MovieSceneCapture/Public/MovieSceneCapture.h
--- a/Engine/Source/Runtime/MovieSceneCapture/Public/MovieSceneCapture.h
+++ b/Engine/Source/Runtime/MovieSceneCapture/Public/MovieSceneCapture.h
@@ -33,23 +33,23 @@ struct FCachedMetrics
 };
 
 /** Class responsible for capturing scene data */
-UCLASS(config=EditorPerProjectUserSettings, PerObjectConfig, BlueprintType)
-class MOVIESCENECAPTURE_API UMovieSceneCapture : public UObject, public IMovieSceneCaptureInterface, public ICaptureProtocolHost
+UCLASS(config=EditorPerProjectUserSettings, PerObjectConfig, BlueprintType, MinimalAPI)
+class UMovieSceneCapture : public UObject, public IMovieSceneCaptureInterface, public ICaptureProtocolHost
 {
 public:
-	UMovieSceneCapture(const FObjectInitializer& Initializer);
+	MOVIESCENECAPTURE_API UMovieSceneCapture(const FObjectInitializer& Initializer);
 
 	GENERATED_BODY()
 
 	/** This name is used by the UI to save/load a specific instance of the settings from config that doesn't affect the CDO which would affect scripting environments. */
-	static const FName MovieSceneCaptureUIName;
+	static MOVIESCENECAPTURE_API const FName MovieSceneCaptureUIName;
 
-	virtual void PostInitProperties() override;
+	MOVIESCENECAPTURE_API virtual void PostInitProperties() override;
 
 public:
 
 	// Begin IMovieSceneCaptureInterface
-	virtual void Initialize(TSharedPtr<FSceneViewport> InSceneViewport, int32 PIEInstance = -1) override;
+	MOVIESCENECAPTURE_API virtual void Initialize(TSharedPtr<FSceneViewport> InSceneViewport, int32 PIEInstance = -1) override;
 	virtual void StartCapturing() { StartCapture(); }
 	virtual void Close() override { Finalize(); }
 	virtual FMovieSceneCaptureHandle GetHandle() const override { return Handle; }
@@ -58,14 +58,14 @@ public:
 	// End IMovieSceneCaptureInterface
 
 	/** Load save from config helpers */
-	virtual void LoadFromConfig();
-	virtual void SaveToConfig();
+	MOVIESCENECAPTURE_API virtual void LoadFromConfig();
+	MOVIESCENECAPTURE_API virtual void SaveToConfig();
 
 	/** Serialize additional json data for this capture */
-	void SerializeJson(FJsonObject& Object);
+	MOVIESCENECAPTURE_API void SerializeJson(FJsonObject& Object);
 
 	/** Deserialize additional json data for this capture */
-	void DeserializeJson(const FJsonObject& Object);
+	MOVIESCENECAPTURE_API void DeserializeJson(const FJsonObject& Object);
 
 protected:
 
@@ -76,7 +76,7 @@ protected:
 	virtual void DeserializeAdditionalJson(const FJsonObject& Object){}
 
 	/** Returns true if this is currently the audio pass, or if an audio pass is not needed. Shorthand for checking if we're in a state where we should finish capture. */
-	virtual bool IsAudioPassIfNeeded() const;
+	MOVIESCENECAPTURE_API virtual bool IsAudioPassIfNeeded() const;
 public:
 
 	/** The type of capture protocol to use for image data */
@@ -131,39 +131,39 @@ public:
 	
 	
 	UFUNCTION(BlueprintCallable, Category=Capture)
-	void SetImageCaptureProtocolType(TSubclassOf<UMovieSceneCaptureProtocolBase> ProtocolType);
+	MOVIESCENECAPTURE_API void SetImageCaptureProtocolType(TSubclassOf<UMovieSceneCaptureProtocolBase> ProtocolType);
 	UFUNCTION(BlueprintCallable, Category=Capture)
-	void SetAudioCaptureProtocolType(TSubclassOf<UMovieSceneCaptureProtocolBase> ProtocolType);
+	MOVIESCENECAPTURE_API void SetAudioCaptureProtocolType(TSubclassOf<UMovieSceneCaptureProtocolBase> ProtocolType);
 
 public:
 
 	/** Starts warming up.  May be optionally called before StartCapture().  This can be used to start rendering frames early, before
 	    any files are captured or written out */
-	void StartWarmup();
+	MOVIESCENECAPTURE_API void StartWarmup();
 
 	/** Initialize the capture so that it is able to start capturing frames */
-	void StartCapture();
+	MOVIESCENECAPTURE_API void StartCapture();
 
 	/** Indicate that this frame should be captured - must be called before the movie scene capture is ticked */
-	void CaptureThisFrame(float DeltaSeconds);
+	MOVIESCENECAPTURE_API void CaptureThisFrame(float DeltaSeconds);
 
 	/** Automatically finalizes the capture when all currently pending frames are dealt with */
-	void FinalizeWhenReady();
+	MOVIESCENECAPTURE_API void FinalizeWhenReady();
 
 	/** Check whether we should automatically finalize this capture */
-	bool ShouldFinalize() const;
+	MOVIESCENECAPTURE_API bool ShouldFinalize() const;
 
 	/** Finalize the capturing process, assumes all frames have been processed. */
-	void Finalize();
+	MOVIESCENECAPTURE_API void Finalize();
 
 public:
 
 	/** Called at the end of a frame, before a frame is presented by slate */
-	void Tick(float DeltaSeconds);
+	MOVIESCENECAPTURE_API void Tick(float DeltaSeconds);
 
 	// ICaptureProtocolHost interface
 	/** Resolve the specified format using the user supplied formatting rules. */
-	FString ResolveFileFormat(const FString& Format, const FFrameMetrics& FrameMetrics) const;
+	MOVIESCENECAPTURE_API FString ResolveFileFormat(const FString& Format, const FFrameMetrics& FrameMetrics) const;
 
 	/** Estimate how long our duration is going to be for pre-allocation purposes. */
 	double GetEstimatedCaptureDurationSeconds() const { return 0.0; }
@@ -179,16 +179,16 @@ protected:
 
 
 	/** Initialize the settings structure for the current capture type */
-	void InitializeCaptureProtocols();
+	MOVIESCENECAPTURE_API void InitializeCaptureProtocols();
 	
-	void ForciblyReinitializeCaptureProtocols();
+	MOVIESCENECAPTURE_API void ForciblyReinitializeCaptureProtocols();
 
 	/** Called at the end of a frame, before a frame is presented by slate */
 	virtual void OnTick(float DeltaSeconds) { CaptureThisFrame(DeltaSeconds); }
 protected:
 
 #if WITH_EDITOR
-	virtual void PostEditChangeProperty( struct FPropertyChangedEvent& PropertyChangedEvent) override;
+	MOVIESCENECAPTURE_API virtual void PostEditChangeProperty( struct FPropertyChangedEvent& PropertyChangedEvent) override;
 #endif
 
 protected:
@@ -219,29 +219,29 @@ protected:
 };
 
 /** A strategy that employs a fixed frame time-step, and as such never drops a frame. Potentially accelerated. */
-struct MOVIESCENECAPTURE_API FFixedTimeStepCaptureStrategy : ICaptureStrategy
+struct FFixedTimeStepCaptureStrategy : ICaptureStrategy
 {
-	FFixedTimeStepCaptureStrategy(FFrameRate InFrameRate);
+	MOVIESCENECAPTURE_API FFixedTimeStepCaptureStrategy(FFrameRate InFrameRate);
 
-	virtual void OnInitialize() override;
-	virtual void OnStop() override;
-	virtual bool ShouldPresent(double CurrentTimeSeconds, uint32 FrameIndex) const override;
-	virtual int32 GetDroppedFrames(double CurrentTimeSeconds, uint32 FrameIndex) const override;
+	MOVIESCENECAPTURE_API virtual void OnInitialize() override;
+	MOVIESCENECAPTURE_API virtual void OnStop() override;
+	MOVIESCENECAPTURE_API virtual bool ShouldPresent(double CurrentTimeSeconds, uint32 FrameIndex) const override;
+	MOVIESCENECAPTURE_API virtual int32 GetDroppedFrames(double CurrentTimeSeconds, uint32 FrameIndex) const override;
 
 private:
 	FFrameRate FrameRate;
 };
 
 /** A capture strategy that captures in real-time, potentially dropping frames to maintain a stable constant framerate video. */
-struct MOVIESCENECAPTURE_API FRealTimeCaptureStrategy : ICaptureStrategy
+struct FRealTimeCaptureStrategy : ICaptureStrategy
 {
-	FRealTimeCaptureStrategy(FFrameRate InFrameRate);
+	MOVIESCENECAPTURE_API FRealTimeCaptureStrategy(FFrameRate InFrameRate);
 
-	virtual void OnInitialize() override;
-	virtual void OnStop() override;
+	MOVIESCENECAPTURE_API virtual void OnInitialize() override;
+	MOVIESCENECAPTURE_API virtual void OnStop() override;
 	virtual bool ShouldSynchronizeFrames() const override { return false; }
-	virtual bool ShouldPresent(double CurrentTimeSeconds, uint32 FrameIndex) const override;
-	virtual int32 GetDroppedFrames(double CurrentTimeSeconds, uint32 FrameIndex) const override;
+	MOVIESCENECAPTURE_API virtual bool ShouldPresent(double CurrentTimeSeconds, uint32 FrameIndex) const override;
+	MOVIESCENECAPTURE_API virtual int32 GetDroppedFrames(double CurrentTimeSeconds, uint32 FrameIndex) const override;
 
 private:
 	double NextPresentTimeS, FrameLength;
diff --git a/Engine/Source/Runtime/MovieSceneCapture/Public/MovieSceneCaptureEnvironment.h b/Engine/Source/Runtime/MovieSceneCapture/Public/MovieSceneCaptureEnvironment.h
--- a/Engine/Source/Runtime/MovieSceneCapture/Public/MovieSceneCaptureEnvironment.h
+++ b/Engine/Source/Runtime/MovieSceneCapture/Public/MovieSceneCaptureEnvironment.h
@@ -10,36 +10,36 @@
 
 class UMovieSceneCaptureProtocolBase;
 
-UCLASS()
-class MOVIESCENECAPTURE_API UMovieSceneCaptureEnvironment : public UObject
+UCLASS(MinimalAPI)
+class UMovieSceneCaptureEnvironment : public UObject
 {
 public:
 	GENERATED_BODY()
 
 	/** Get the frame number of the current capture */
 	UFUNCTION(BlueprintPure, Category="Cinematics|Capture")
-	static int32 GetCaptureFrameNumber();
+	static MOVIESCENECAPTURE_API int32 GetCaptureFrameNumber();
 
 	/** Get the total elapsed time of the current capture in seconds */
 	UFUNCTION(BlueprintPure, Category="Cinematics|Capture")
-	static float GetCaptureElapsedTime();
+	static MOVIESCENECAPTURE_API float GetCaptureElapsedTime();
 
 	/**
 	 * Return true if there is any capture currently active (even in a warm-up state).
 	 * Useful for checking whether to do certain operations in BeginPlay
 	 */
 	UFUNCTION(BlueprintCallable, Category="Cinematics|Capture")
-	static bool IsCaptureInProgress();
+	static MOVIESCENECAPTURE_API bool IsCaptureInProgress();
 
 	/**
 	 * Attempt to locate a capture protocol - may not be in a capturing state
 	 */
 	UFUNCTION(BlueprintCallable, Category="Cinematics|Capture")
-	static UMovieSceneImageCaptureProtocolBase* FindImageCaptureProtocol();
+	static MOVIESCENECAPTURE_API UMovieSceneImageCaptureProtocolBase* FindImageCaptureProtocol();
 
 	/**
 	* Attempt to locate a capture protocol - may not be in a capturing state
 	*/
 	UFUNCTION(BlueprintCallable, Category = "Cinematics|Capture")
-	static UMovieSceneAudioCaptureProtocolBase* FindAudioCaptureProtocol();
+	static MOVIESCENECAPTURE_API UMovieSceneAudioCaptureProtocolBase* FindAudioCaptureProtocol();
 };
diff --git a/Engine/Source/Runtime/MovieSceneCapture/Public/MovieSceneCaptureProtocolBase.h b/Engine/Source/Runtime/MovieSceneCapture/Public/MovieSceneCaptureProtocolBase.h
--- a/Engine/Source/Runtime/MovieSceneCapture/Public/MovieSceneCaptureProtocolBase.h
+++ b/Engine/Source/Runtime/MovieSceneCapture/Public/MovieSceneCaptureProtocolBase.h
@@ -53,14 +53,14 @@ private:
  * A typical process for capture consits of the following process:
  *     Setup -> [ Warm up -> [ Capture Frame ] ] -> Begin Finalize -> [ HasFinishedProcessing ] -> Finalize
  */
-UCLASS(config=EditorPerProjectUserSettings, PerObjectConfig, Abstract)
-class MOVIESCENECAPTURE_API UMovieSceneCaptureProtocolBase : public UObject
+UCLASS(config=EditorPerProjectUserSettings, PerObjectConfig, Abstract, MinimalAPI)
+class UMovieSceneCaptureProtocolBase : public UObject
 {
 public:
 
 	GENERATED_BODY()
 
-	UMovieSceneCaptureProtocolBase(const FObjectInitializer& ObjInit);
+	MOVIESCENECAPTURE_API UMovieSceneCaptureProtocolBase(const FObjectInitializer& ObjInit);
 
 public:
 
@@ -82,37 +82,37 @@ public:
 	 * @param InSettings		The initial initialization settings to use for the capture
 	 * @param Host				The client that is initializing this protocol
 	 */
-	bool Setup(const FCaptureProtocolInitSettings& InSettings, const ICaptureProtocolHost* Host);
+	MOVIESCENECAPTURE_API bool Setup(const FCaptureProtocolInitSettings& InSettings, const ICaptureProtocolHost* Host);
 
 	/**
 	 * Get the UWorld associated with this Capture Protocol. This is not valid until
 	 * Setup has been called with a valid Slate viewport. Will return nullptr when
 	 * the protocol has been created but the game world is not running (ie: in UI).
 	 */
-	virtual class UWorld* GetWorld() const override;
+	MOVIESCENECAPTURE_API virtual class UWorld* GetWorld() const override;
 public:
 
 	/**
 	 * Called on the main thread before the movie capture itself is updated to reset per-frame state
 	 */
-	void PreTick();
+	MOVIESCENECAPTURE_API void PreTick();
 
 	/**
 	 * Called on the main thread to do any additional processing
 	 */
-	void Tick();
+	MOVIESCENECAPTURE_API void Tick();
 
 	/**
 	 * Start warming up this capture protocol - called any time the process enters a warming-up state
 	 */
-	void WarmUp();
+	MOVIESCENECAPTURE_API void WarmUp();
 
 	/**
 	 * Called when this protocol should start capturing
 	 *
 	 * @return true if the operation was successful, false otherwise
 	 */
-	bool StartCapture();
+	MOVIESCENECAPTURE_API bool StartCapture();
 
 	/**
 	 * Instruct this protocol to capture a frame relating to the specified metrics
@@ -120,40 +120,40 @@ public:
 	 * @param FrameMetrics		Frame metrics relating to the current frame
 	 * @param Host				The client that is initializing this protocol
 	 */
-	void CaptureFrame(const FFrameMetrics& FrameMetrics);
+	MOVIESCENECAPTURE_API void CaptureFrame(const FFrameMetrics& FrameMetrics);
 
 	/**
 	 * Called when we have finished capturing and we should start finalizing the capture
 	 */
-	void BeginFinalize();
+	MOVIESCENECAPTURE_API void BeginFinalize();
 
 	/**
 	 * Check whether this protocol has any processing left to do, or whether it should be finalized.
 	 * Only called when the capture has been asked to end.
 	 */
-	bool HasFinishedProcessing() const;
+	MOVIESCENECAPTURE_API bool HasFinishedProcessing() const;
 
 	/**
 	 * Called when this protocol should tear down and finalize all its processing. Only called if HasFinishedProcessing is true.
 	 */
-	void Finalize();
+	MOVIESCENECAPTURE_API void Finalize();
 
 	/**
 	 * Called when generating formatting filename to add any additional format mappings
 	 *
 	 * @param FormatMappings	Map to add additional format rules to
 	 */
-	void AddFormatMappings(TMap<FString, FStringFormatArg>& FormatMappings) const;
+	MOVIESCENECAPTURE_API void AddFormatMappings(TMap<FString, FStringFormatArg>& FormatMappings) const;
 
 	/**
 	 * Called when this protocol has been released
 	 */
-	void OnReleaseConfig(FMovieSceneCaptureSettings& InSettings);
+	MOVIESCENECAPTURE_API void OnReleaseConfig(FMovieSceneCaptureSettings& InSettings);
 
 	/**
 	 * Called when this protocol has been loaded
 	 */
-	void OnLoadConfig(FMovieSceneCaptureSettings& InSettings);
+	MOVIESCENECAPTURE_API void OnLoadConfig(FMovieSceneCaptureSettings& InSettings);
 
 	/**
 	 * Test whether this capture protocol thinks the file should be written to. Only called when we're not overwriting existing files.
@@ -163,7 +163,7 @@ public:
 	 * @param bOverwriteExisting	Whether we are allowed to overwrite existing files
 	 * @return Whether we should deem this file writable or not
 	 */
-	bool CanWriteToFile(const TCHAR* InFilename, bool bOverwriteExisting) const;
+	MOVIESCENECAPTURE_API bool CanWriteToFile(const TCHAR* InFilename, bool bOverwriteExisting) const;
 
 protected:
 
@@ -253,9 +253,9 @@ protected:
 	 * @param bOverwriteExisting	Whether we are allowed to overwrite existing files
 	 * @return Whether we should deem this file writable or not
 	 */
-	virtual bool CanWriteToFileImpl(const TCHAR* InFilename, bool bOverwriteExisting) const;
-	virtual FString GenerateFilenameImpl(const FFrameMetrics& FrameMetrics, const TCHAR* Extension, const FString* NameFormatString = nullptr) const;
-	void EnsureFileWritableImpl(const FString& File) const;
+	MOVIESCENECAPTURE_API virtual bool CanWriteToFileImpl(const TCHAR* InFilename, bool bOverwriteExisting) const;
+	MOVIESCENECAPTURE_API virtual FString GenerateFilenameImpl(const FFrameMetrics& FrameMetrics, const TCHAR* Extension, const FString* NameFormatString = nullptr) const;
+	MOVIESCENECAPTURE_API void EnsureFileWritableImpl(const FString& File) const;
 
 protected:
 
@@ -278,8 +278,8 @@ private:
 /**
 * A class to inherit from for image capture protocols. Used to filter the UI for protocols used on the image capture pass.
 */
-UCLASS(config = EditorPerProjectUserSettings, PerObjectConfig, Abstract)
-class MOVIESCENECAPTURE_API  UMovieSceneImageCaptureProtocolBase : public UMovieSceneCaptureProtocolBase
+UCLASS(config = EditorPerProjectUserSettings, PerObjectConfig, Abstract, MinimalAPI)
+class  UMovieSceneImageCaptureProtocolBase : public UMovieSceneCaptureProtocolBase
 {
 	GENERATED_BODY()
 
@@ -293,8 +293,8 @@ public:
 /**
 * A class to inherit from for audio capture protocols. Used to filter the UI for protocols used on the audio capture pass.
 */
-UCLASS(config = EditorPerProjectUserSettings, PerObjectConfig, Abstract)
-class MOVIESCENECAPTURE_API UMovieSceneAudioCaptureProtocolBase : public UMovieSceneCaptureProtocolBase
+UCLASS(config = EditorPerProjectUserSettings, PerObjectConfig, Abstract, MinimalAPI)
+class UMovieSceneAudioCaptureProtocolBase : public UMovieSceneCaptureProtocolBase
 {
 	GENERATED_BODY()
 
diff --git a/Engine/Source/Runtime/MovieSceneCapture/Public/MovieSceneCaptureSettings.h b/Engine/Source/Runtime/MovieSceneCapture/Public/MovieSceneCaptureSettings.h
--- a/Engine/Source/Runtime/MovieSceneCapture/Public/MovieSceneCaptureSettings.h
+++ b/Engine/Source/Runtime/MovieSceneCapture/Public/MovieSceneCaptureSettings.h
@@ -12,7 +12,7 @@
 
 /** Structure representing a capture resolution */
 USTRUCT(BlueprintType)
-struct MOVIESCENECAPTURE_API FCaptureResolution
+struct FCaptureResolution
 {
 	FCaptureResolution(int32 InX = 0, int32 InY = 0) : ResX(InX), ResY(InY) {}
 	
@@ -27,9 +27,9 @@ struct MOVIESCENECAPTURE_API FCaptureResolution
 
 /** Common movie-scene capture settings */
 USTRUCT(BlueprintType)
-struct MOVIESCENECAPTURE_API FMovieSceneCaptureSettings
+struct FMovieSceneCaptureSettings
 {
-	FMovieSceneCaptureSettings();
+	MOVIESCENECAPTURE_API FMovieSceneCaptureSettings();
 
 	GENERATED_BODY()
 
diff --git a/Engine/Source/Runtime/MovieSceneCapture/Public/Protocols/AudioCaptureProtocol.h b/Engine/Source/Runtime/MovieSceneCapture/Public/Protocols/AudioCaptureProtocol.h
--- a/Engine/Source/Runtime/MovieSceneCapture/Public/Protocols/AudioCaptureProtocol.h
+++ b/Engine/Source/Runtime/MovieSceneCapture/Public/Protocols/AudioCaptureProtocol.h
@@ -15,8 +15,8 @@
  * This is a null audio capture implementation which skips capturing audio. The MovieSceneCapture is explicitly
  * aware of this type and will skip processing an audio pass if this is specified.
  */
-UCLASS(meta = (DisplayName = "No Audio", CommandLineID = "NullAudio"))
-class MOVIESCENECAPTURE_API UNullAudioCaptureProtocol : public UMovieSceneAudioCaptureProtocolBase
+UCLASS(meta = (DisplayName = "No Audio", CommandLineID = "NullAudio"), MinimalAPI)
+class UNullAudioCaptureProtocol : public UMovieSceneAudioCaptureProtocolBase
 {
 	GENERATED_BODY()
 };
diff --git a/Engine/Source/Runtime/MovieSceneCapture/Public/Protocols/CompositionGraphCaptureProtocol.h b/Engine/Source/Runtime/MovieSceneCapture/Public/Protocols/CompositionGraphCaptureProtocol.h
--- a/Engine/Source/Runtime/MovieSceneCapture/Public/Protocols/CompositionGraphCaptureProtocol.h
+++ b/Engine/Source/Runtime/MovieSceneCapture/Public/Protocols/CompositionGraphCaptureProtocol.h
@@ -37,7 +37,7 @@ static_assert(HCGM_ACEScg == (int32)EDisplayColorGamut::ACEScg_D60, "EHDRCapture
 static_assert(HCGM_MAX == (int32)EDisplayColorGamut::MAX + 1, "EHDRCaptureGamut and EDisplayColorGamut not matching");
 
 USTRUCT(BlueprintType)
-struct MOVIESCENECAPTURE_API FCompositionGraphCapturePasses
+struct FCompositionGraphCapturePasses
 {
 	GENERATED_BODY()
 
@@ -46,8 +46,8 @@ struct MOVIESCENECAPTURE_API FCompositionGraphCapturePasses
 	TArray<FString> Value;
 };
 
-UCLASS(config=EditorPerProjectUserSettings, meta=(DisplayName="Custom Render Passes", CommandLineID="CustomRenderPasses"))
-class MOVIESCENECAPTURE_API UCompositionGraphCaptureProtocol : public UMovieSceneImageCaptureProtocolBase
+UCLASS(config=EditorPerProjectUserSettings, meta=(DisplayName="Custom Render Passes", CommandLineID="CustomRenderPasses"), MinimalAPI)
+class UCompositionGraphCaptureProtocol : public UMovieSceneImageCaptureProtocolBase
 {
 public:
 	GENERATED_BODY()
@@ -84,13 +84,13 @@ public:
 public:
 
 	/**~ UMovieSceneCaptureProtocolBase implementation */
-	virtual bool SetupImpl();
-	virtual void CaptureFrameImpl(const FFrameMetrics& FrameMetrics);
-	virtual void TickImpl() override;
-	virtual void FinalizeImpl() override;
-	virtual bool HasFinishedProcessingImpl() const override;
-	virtual void OnReleaseConfigImpl(FMovieSceneCaptureSettings& InSettings) override;
-	virtual void OnLoadConfigImpl(FMovieSceneCaptureSettings& InSettings) override;
+	MOVIESCENECAPTURE_API virtual bool SetupImpl();
+	MOVIESCENECAPTURE_API virtual void CaptureFrameImpl(const FFrameMetrics& FrameMetrics);
+	MOVIESCENECAPTURE_API virtual void TickImpl() override;
+	MOVIESCENECAPTURE_API virtual void FinalizeImpl() override;
+	MOVIESCENECAPTURE_API virtual bool HasFinishedProcessingImpl() const override;
+	MOVIESCENECAPTURE_API virtual void OnReleaseConfigImpl(FMovieSceneCaptureSettings& InSettings) override;
+	MOVIESCENECAPTURE_API virtual void OnLoadConfigImpl(FMovieSceneCaptureSettings& InSettings) override;
 	/**~ End UMovieSceneCaptureProtocolBase implementation */
 
 private:
diff --git a/Engine/Source/Runtime/MovieSceneCapture/Public/Protocols/FrameGrabberProtocol.h b/Engine/Source/Runtime/MovieSceneCapture/Public/Protocols/FrameGrabberProtocol.h
--- a/Engine/Source/Runtime/MovieSceneCapture/Public/Protocols/FrameGrabberProtocol.h
+++ b/Engine/Source/Runtime/MovieSceneCapture/Public/Protocols/FrameGrabberProtocol.h
@@ -10,8 +10,8 @@
 #include "FrameGrabberProtocol.generated.h"
 
 
-UCLASS(Abstract, config=EditorPerProjectUserSettings)
-class MOVIESCENECAPTURE_API UFrameGrabberProtocol : public UMovieSceneImageCaptureProtocolBase
+UCLASS(Abstract, config=EditorPerProjectUserSettings, MinimalAPI)
+class UFrameGrabberProtocol : public UMovieSceneImageCaptureProtocolBase
 {
 public:
 
@@ -32,12 +32,12 @@ public:
 protected:
 
 	/**~ UMovieSceneCaptureProtocolBase Implementation */
-	virtual bool HasFinishedProcessingImpl() const override;
-	virtual bool SetupImpl() override;
-	virtual void CaptureFrameImpl(const FFrameMetrics& FrameMetrics);
-	virtual void TickImpl() override;
-	virtual void BeginFinalizeImpl() override;
-	virtual void FinalizeImpl() override;
+	MOVIESCENECAPTURE_API virtual bool HasFinishedProcessingImpl() const override;
+	MOVIESCENECAPTURE_API virtual bool SetupImpl() override;
+	MOVIESCENECAPTURE_API virtual void CaptureFrameImpl(const FFrameMetrics& FrameMetrics);
+	MOVIESCENECAPTURE_API virtual void TickImpl() override;
+	MOVIESCENECAPTURE_API virtual void BeginFinalizeImpl() override;
+	MOVIESCENECAPTURE_API virtual void FinalizeImpl() override;
 	/**~ End UMovieSceneCaptureProtocolBase Implementation */
 
 protected:
@@ -51,14 +51,14 @@ protected:
 
 	 * @return Shared pointer to a payload to associate with the frame, or nullptr
 	 */
-	virtual FFramePayloadPtr GetFramePayload(const FFrameMetrics& FrameMetrics) PURE_VIRTUAL(UFrameGrabberProtocol::GetFramePayload, return nullptr;)
+	MOVIESCENECAPTURE_API virtual FFramePayloadPtr GetFramePayload(const FFrameMetrics& FrameMetrics) PURE_VIRTUAL(UFrameGrabberProtocol::GetFramePayload, return nullptr;)
 
 	/**
 	 * Process a captured frame. This may be called on any thread.
 	 *
 	 * @param Frame			The captured frame data, including any payload retrieved from GetFramePayload
 	 */
-	virtual void ProcessFrame(FCapturedFrameData Frame) PURE_VIRTUAL(UFrameGrabberProtocol::ProcessFrame,)
+	MOVIESCENECAPTURE_API virtual void ProcessFrame(FCapturedFrameData Frame) PURE_VIRTUAL(UFrameGrabberProtocol::ProcessFrame,)
 
 private:
 
diff --git a/Engine/Source/Runtime/MovieSceneCapture/Public/Protocols/ImageSequenceProtocol.h b/Engine/Source/Runtime/MovieSceneCapture/Public/Protocols/ImageSequenceProtocol.h
--- a/Engine/Source/Runtime/MovieSceneCapture/Public/Protocols/ImageSequenceProtocol.h
+++ b/Engine/Source/Runtime/MovieSceneCapture/Public/Protocols/ImageSequenceProtocol.h
@@ -13,30 +13,30 @@
 struct FMovieSceneCaptureSettings;
 class IImageWriteQueue;
 
-UCLASS(Abstract, config=EditorPerProjectUserSettings)
-class MOVIESCENECAPTURE_API UImageSequenceProtocol : public UFrameGrabberProtocol
+UCLASS(Abstract, config=EditorPerProjectUserSettings, MinimalAPI)
+class UImageSequenceProtocol : public UFrameGrabberProtocol
 {
 public:
 
 	GENERATED_BODY()
 
-	UImageSequenceProtocol(const FObjectInitializer& ObjInit);
+	MOVIESCENECAPTURE_API UImageSequenceProtocol(const FObjectInitializer& ObjInit);
 
 public:
 
 	/** ~UFrameGrabberProtocol implementation */
-	virtual FFramePayloadPtr GetFramePayload(const FFrameMetrics& FrameMetrics);
-	virtual void ProcessFrame(FCapturedFrameData Frame);
+	MOVIESCENECAPTURE_API virtual FFramePayloadPtr GetFramePayload(const FFrameMetrics& FrameMetrics);
+	MOVIESCENECAPTURE_API virtual void ProcessFrame(FCapturedFrameData Frame);
 	/** ~End UFrameGrabberProtocol implementation */
 
 	/** ~ UMovieSceneCaptureProtocol implementation */
-	virtual bool SetupImpl() override;
-	virtual void AddFormatMappingsImpl(TMap<FString, FStringFormatArg>& FormatMappings) const override;
-	virtual void BeginFinalizeImpl() override;
-	virtual void FinalizeImpl() override;
-	virtual bool HasFinishedProcessingImpl() const override;
-	virtual void OnReleaseConfigImpl(FMovieSceneCaptureSettings& InSettings) override;
-	virtual void OnLoadConfigImpl(FMovieSceneCaptureSettings& InSettings) override;
+	MOVIESCENECAPTURE_API virtual bool SetupImpl() override;
+	MOVIESCENECAPTURE_API virtual void AddFormatMappingsImpl(TMap<FString, FStringFormatArg>& FormatMappings) const override;
+	MOVIESCENECAPTURE_API virtual void BeginFinalizeImpl() override;
+	MOVIESCENECAPTURE_API virtual void FinalizeImpl() override;
+	MOVIESCENECAPTURE_API virtual bool HasFinishedProcessingImpl() const override;
+	MOVIESCENECAPTURE_API virtual void OnReleaseConfigImpl(FMovieSceneCaptureSettings& InSettings) override;
+	MOVIESCENECAPTURE_API virtual void OnLoadConfigImpl(FMovieSceneCaptureSettings& InSettings) override;
 	/** ~End UMovieSceneCaptureProtocol implementation */
 
 protected:
@@ -59,8 +59,8 @@ private:
 	TFuture<void> FinalizeFence;
 };
 
-UCLASS(config=EditorPerProjectUserSettings, Abstract)
-class MOVIESCENECAPTURE_API UCompressedImageSequenceProtocol : public UImageSequenceProtocol
+UCLASS(config=EditorPerProjectUserSettings, Abstract, MinimalAPI)
+class UCompressedImageSequenceProtocol : public UImageSequenceProtocol
 {
 public:
 
@@ -79,12 +79,12 @@ public:
 protected:
 
 	virtual int32 GetCompressionQuality() const override { return CompressionQuality; }
-	virtual bool SetupImpl() override;
-	virtual void AddFormatMappingsImpl(TMap<FString, FStringFormatArg>& FormatMappings) const override;
+	MOVIESCENECAPTURE_API virtual bool SetupImpl() override;
+	MOVIESCENECAPTURE_API virtual void AddFormatMappingsImpl(TMap<FString, FStringFormatArg>& FormatMappings) const override;
 };
 
-UCLASS(meta=(DisplayName="Image Sequence (bmp)", CommandLineID="BMP"))
-class MOVIESCENECAPTURE_API UImageSequenceProtocol_BMP : public UImageSequenceProtocol
+UCLASS(meta=(DisplayName="Image Sequence (bmp)", CommandLineID="BMP"), MinimalAPI)
+class UImageSequenceProtocol_BMP : public UImageSequenceProtocol
 {
 public:
 
@@ -97,8 +97,8 @@ public:
 	}
 };
 
-UCLASS(meta=(DisplayName="Image Sequence (png)", CommandLineID="PNG"))
-class MOVIESCENECAPTURE_API UImageSequenceProtocol_PNG : public UCompressedImageSequenceProtocol
+UCLASS(meta=(DisplayName="Image Sequence (png)", CommandLineID="PNG"), MinimalAPI)
+class UImageSequenceProtocol_PNG : public UCompressedImageSequenceProtocol
 {
 public:
 
@@ -111,8 +111,8 @@ public:
 	}
 };
 
-UCLASS(meta=(DisplayName="Image Sequence (jpg)", CommandLineID="JPG"))
-class MOVIESCENECAPTURE_API UImageSequenceProtocol_JPG : public UCompressedImageSequenceProtocol
+UCLASS(meta=(DisplayName="Image Sequence (jpg)", CommandLineID="JPG"), MinimalAPI)
+class UImageSequenceProtocol_JPG : public UCompressedImageSequenceProtocol
 {
 public:
 
@@ -125,8 +125,8 @@ public:
 	}
 };
 
-UCLASS(meta=(DisplayName="Image Sequence (exr)", CommandLineID="EXR"))
-class MOVIESCENECAPTURE_API UImageSequenceProtocol_EXR : public UImageSequenceProtocol
+UCLASS(meta=(DisplayName="Image Sequence (exr)", CommandLineID="EXR"), MinimalAPI)
+class UImageSequenceProtocol_EXR : public UImageSequenceProtocol
 {
 public:
 
@@ -140,13 +140,13 @@ public:
 	UPROPERTY(config, EditAnywhere, Category=ImageSettings)
 	TEnumAsByte<EHDRCaptureGamut> CaptureGamut;
 
-	UImageSequenceProtocol_EXR(const FObjectInitializer& ObjInit);
+	MOVIESCENECAPTURE_API UImageSequenceProtocol_EXR(const FObjectInitializer& ObjInit);
 
 private:
 	virtual int32 GetCompressionQuality() const override { return bCompressed ? (int32)EImageCompressionQuality::Default : (int32)EImageCompressionQuality::Uncompressed; }
-	virtual bool SetupImpl() override;
-	virtual void FinalizeImpl() override;
-	virtual void AddFormatMappingsImpl(TMap<FString, FStringFormatArg>& FormatMappings) const override;
+	MOVIESCENECAPTURE_API virtual bool SetupImpl() override;
+	MOVIESCENECAPTURE_API virtual void FinalizeImpl() override;
+	MOVIESCENECAPTURE_API virtual void AddFormatMappingsImpl(TMap<FString, FStringFormatArg>& FormatMappings) const override;
 
 	EDisplayColorGamut RestoreColorGamut;
 	EDisplayOutputFormat RestoreOutputDevice;
diff --git a/Engine/Source/Runtime/MovieSceneCapture/Public/Protocols/UserDefinedCaptureProtocol.h b/Engine/Source/Runtime/MovieSceneCapture/Public/Protocols/UserDefinedCaptureProtocol.h
--- a/Engine/Source/Runtime/MovieSceneCapture/Public/Protocols/UserDefinedCaptureProtocol.h
+++ b/Engine/Source/Runtime/MovieSceneCapture/Public/Protocols/UserDefinedCaptureProtocol.h
@@ -41,15 +41,15 @@ struct FCapturedPixels
 /**
  * A blueprintable capture protocol that defines how to capture frames from the engine
  */
-UCLASS(Abstract, Blueprintable, meta=(DisplayName="Capture Protocol"))
-class MOVIESCENECAPTURE_API UUserDefinedCaptureProtocol
+UCLASS(Abstract, Blueprintable, meta=(DisplayName="Capture Protocol"), MinimalAPI)
+class UUserDefinedCaptureProtocol
 	: public UMovieSceneImageCaptureProtocolBase
 {
 public:
 
 	GENERATED_BODY()
 
-	UUserDefinedCaptureProtocol(const FObjectInitializer& ObjInit);
+	MOVIESCENECAPTURE_API UUserDefinedCaptureProtocol(const FObjectInitializer& ObjInit);
 
 protected:
 
@@ -58,58 +58,58 @@ protected:
 	 * Useful for any pre-frame set up or resetting the previous frame's state
 	 */
 	UFUNCTION(BlueprintImplementableEvent, Category=Capture)
-	void OnPreTick();
+	MOVIESCENECAPTURE_API void OnPreTick();
 
 	/**
 	 * Called after the capture process itself is ticked, after the frame is set up for capture, but before most actors have ticked
 	 */
 	UFUNCTION(BlueprintImplementableEvent, Category=Capture)
-	void OnTick();
+	MOVIESCENECAPTURE_API void OnTick();
 
 	/**
 	 * Called once at the start of the capture process, but before any warmup frames
 	 * @return true if this protocol set up successfully, false otherwise
 	 */
 	UFUNCTION(BlueprintNativeEvent, Category=Capture)
-	bool OnSetup();
+	MOVIESCENECAPTURE_API bool OnSetup();
 	virtual bool OnSetup_Implementation() { return true; }
 
 	/**
 	 * Called when the capture process is warming up. Protocols may transition from either an initialized, or capturing state to warm-up
 	 */
 	UFUNCTION(BlueprintImplementableEvent, Category=Capture)
-	void OnWarmUp();
+	MOVIESCENECAPTURE_API void OnWarmUp();
 
 	/**
 	 * Called when this protocol should start capturing frames
 	 */
 	UFUNCTION(BlueprintImplementableEvent, Category=Capture)
-	void OnStartCapture();
+	MOVIESCENECAPTURE_API void OnStartCapture();
 
 	/**
 	 * Called when this protocol should capture the current frame
 	 */
 	UFUNCTION(BlueprintImplementableEvent, Category=Capture)
-	void OnCaptureFrame();
+	MOVIESCENECAPTURE_API void OnCaptureFrame();
 
 	/**
 	 * Called when this protocol should temporarily stop capturing frames
 	 */
 	UFUNCTION(BlueprintImplementableEvent, Category=Capture)
-	void OnPauseCapture();
+	MOVIESCENECAPTURE_API void OnPauseCapture();
 
 	/**
 	 * Called when this protocol is going to be shut down - should not capture any more frames
 	 */
 	UFUNCTION(BlueprintImplementableEvent, Category=Capture)
-	void OnBeginFinalize();
+	MOVIESCENECAPTURE_API void OnBeginFinalize();
 
 
 	/**
 	 * Called to check whether this protocol has finished any pending tasks, and can now be shut down
 	 */
 	UFUNCTION(BlueprintNativeEvent, Category=Capture)
-	bool OnCanFinalize() const;
+	MOVIESCENECAPTURE_API bool OnCanFinalize() const;
 	virtual bool OnCanFinalize_Implementation() const { return true; }
 
 
@@ -117,14 +117,14 @@ protected:
 	 * Called to complete finalization of this protocol
 	 */
 	UFUNCTION(BlueprintImplementableEvent, Category=Capture)
-	void OnFinalize();
+	MOVIESCENECAPTURE_API void OnFinalize();
 
 
 	/**
 	 * Called when pixels have been received for the specified stream name
 	 */
 	UFUNCTION(BlueprintImplementableEvent, Category=Capture)
-	void OnPixelsReceived(const FCapturedPixels& Pixels, const FCapturedPixelsID& ID, FFrameMetrics FrameMetrics);
+	MOVIESCENECAPTURE_API void OnPixelsReceived(const FCapturedPixels& Pixels, const FCapturedPixelsID& ID, FFrameMetrics FrameMetrics);
 
 
 public:
@@ -136,7 +136,7 @@ public:
 	 * @param BufferID        The ID of this buffer that is passed to the pixel handler (e.g. a composition pass name).
 	 */
 	UFUNCTION(BlueprintCallable, Category=Capture)
-	void ResolveBuffer(UTexture* Buffer, const FCapturedPixelsID& BufferID);
+	MOVIESCENECAPTURE_API void ResolveBuffer(UTexture* Buffer, const FCapturedPixelsID& BufferID);
 
 
 	/**
@@ -145,21 +145,21 @@ public:
 	 * @param StreamID        The identifier to use for the final pixels buffer
 	 */
 	UFUNCTION(BlueprintCallable, Category=Capture)
-	void StartCapturingFinalPixels(const FCapturedPixelsID& StreamID);
+	MOVIESCENECAPTURE_API void StartCapturingFinalPixels(const FCapturedPixelsID& StreamID);
 
 
 	/**
 	 * Instruct this protocol to stop capturing LDR final pixels
 	 */
 	UFUNCTION(BlueprintCallable, Category=Capture)
-	void StopCapturingFinalPixels();
+	MOVIESCENECAPTURE_API void StopCapturingFinalPixels();
 
 
 	/**
 	 * Generate a filename for the current frame
 	 */
 	UFUNCTION(BlueprintCallable, Category=Capture)
-	virtual FString GenerateFilename(const FFrameMetrics& InFrameMetrics) const;
+	MOVIESCENECAPTURE_API virtual FString GenerateFilename(const FFrameMetrics& InFrameMetrics) const;
 
 
 	/**
@@ -176,27 +176,27 @@ public:
 	/**
 	 * Called when image pixel data is ready to be processed
 	 */
-	void OnPixelsReceivedImpl(const FCapturedPixels& Pixels, const FCapturedPixelsID& StreamID, FFrameMetrics FrameMetrics);
+	MOVIESCENECAPTURE_API void OnPixelsReceivedImpl(const FCapturedPixels& Pixels, const FCapturedPixelsID& StreamID, FFrameMetrics FrameMetrics);
 
 	/**
 	 * INTERNAL: Report that the protocol is dispatching the specified number of asynchronous tasks that need to be completed
 	 * before this protocol can be finalized
 	 */
-	void ReportOutstandingWork(int32 NumNewOperations);
+	MOVIESCENECAPTURE_API void ReportOutstandingWork(int32 NumNewOperations);
 
 protected:
 
 	/**~ Begin UMovieSceneCaptureProtocolBase implementation */
-	virtual void PreTickImpl() override;
-	virtual void TickImpl() override;
-	virtual bool SetupImpl() override;
-	virtual void WarmUpImpl() override;
-	virtual bool StartCaptureImpl() override;
-	virtual void CaptureFrameImpl(const FFrameMetrics& FrameMetrics) override;
-	virtual void BeginFinalizeImpl() override;
-	virtual bool HasFinishedProcessingImpl() const override;
-	virtual void FinalizeImpl() override;
-	virtual void AddFormatMappingsImpl(TMap<FString, FStringFormatArg>& FormatMappings) const override;
+	MOVIESCENECAPTURE_API virtual void PreTickImpl() override;
+	MOVIESCENECAPTURE_API virtual void TickImpl() override;
+	MOVIESCENECAPTURE_API virtual bool SetupImpl() override;
+	MOVIESCENECAPTURE_API virtual void WarmUpImpl() override;
+	MOVIESCENECAPTURE_API virtual bool StartCaptureImpl() override;
+	MOVIESCENECAPTURE_API virtual void CaptureFrameImpl(const FFrameMetrics& FrameMetrics) override;
+	MOVIESCENECAPTURE_API virtual void BeginFinalizeImpl() override;
+	MOVIESCENECAPTURE_API virtual bool HasFinishedProcessingImpl() const override;
+	MOVIESCENECAPTURE_API virtual void FinalizeImpl() override;
+	MOVIESCENECAPTURE_API virtual void AddFormatMappingsImpl(TMap<FString, FStringFormatArg>& FormatMappings) const override;
 	/**~ End UMovieSceneCaptureProtocolBase implementation */
 
 	/**~ Begin UObject implementation */
@@ -231,15 +231,15 @@ protected:
 /**
  * A blueprintable capture protocol tailored to capturing and exporting frames as images
  */
-UCLASS(Abstract, config=EditorPerProjectUserSettings, Blueprintable, meta=(DisplayName="Capture Protocol"))
-class MOVIESCENECAPTURE_API UUserDefinedImageCaptureProtocol
+UCLASS(Abstract, config=EditorPerProjectUserSettings, Blueprintable, meta=(DisplayName="Capture Protocol"), MinimalAPI)
+class UUserDefinedImageCaptureProtocol
 	: public UUserDefinedCaptureProtocol
 {
 public:
 
 	GENERATED_BODY()
 
-	UUserDefinedImageCaptureProtocol(const FObjectInitializer& ObjInit);
+	MOVIESCENECAPTURE_API UUserDefinedImageCaptureProtocol(const FObjectInitializer& ObjInit);
 
 	/** The image format to save as */
 	UPROPERTY(config, BlueprintReadWrite, EditAnywhere, Category=Capture)
@@ -264,7 +264,7 @@ public:
 	 * @return A fully qualified file name
 	 */
 	UFUNCTION(BlueprintCallable, Category=Capture)
-	FString GenerateFilenameForBuffer(UTexture* Buffer, const FCapturedPixelsID& StreamID);
+	MOVIESCENECAPTURE_API FString GenerateFilenameForBuffer(UTexture* Buffer, const FCapturedPixelsID& StreamID);
 
 
 	/*
@@ -273,7 +273,7 @@ public:
 	 * @return A fully qualified file name for the current frame number
 	 */
 	UFUNCTION(BlueprintCallable, Category=Capture)
-	FString GenerateFilenameForCurrentFrame();
+	MOVIESCENECAPTURE_API FString GenerateFilenameForCurrentFrame();
 
 
 	/*
@@ -282,13 +282,13 @@ public:
 	 * @return A fully qualified file name for the current frame number
 	 */
 	UFUNCTION(BlueprintCallable, Category=Capture)
-	void WriteImageToDisk(const FCapturedPixels& PixelData, const FCapturedPixelsID& StreamID, const FFrameMetrics& FrameMetrics, bool bCopyImageData = false);
+	MOVIESCENECAPTURE_API void WriteImageToDisk(const FCapturedPixels& PixelData, const FCapturedPixelsID& StreamID, const FFrameMetrics& FrameMetrics, bool bCopyImageData = false);
 
 protected:
 
 	/**~ Begin UMovieSceneCaptureProtocolBase implementation */
-	virtual void OnReleaseConfigImpl(FMovieSceneCaptureSettings& InSettings) override;
-	virtual void OnLoadConfigImpl(FMovieSceneCaptureSettings& InSettings) override;
+	MOVIESCENECAPTURE_API virtual void OnReleaseConfigImpl(FMovieSceneCaptureSettings& InSettings) override;
+	MOVIESCENECAPTURE_API virtual void OnLoadConfigImpl(FMovieSceneCaptureSettings& InSettings) override;
 	/**~ End UMovieSceneCaptureProtocolBase implementation */
 
 public:
@@ -296,11 +296,11 @@ public:
 	/**
 	 * Called on the main thread when an async operation dispatched from this class has completed (either successfully or otherwise)
 	 */
-	void OnFileWritten();
+	MOVIESCENECAPTURE_API void OnFileWritten();
 
 private:
 
 	/**~ Begin UUserDefinedCaptureProtocol implementation */
-	virtual FString GenerateFilename(const FFrameMetrics& InFrameMetrics) const override;
+	MOVIESCENECAPTURE_API virtual FString GenerateFilename(const FFrameMetrics& InFrameMetrics) const override;
 	/**~ Begin UUserDefinedCaptureProtocol implementation */
-};
\ No newline at end of file
+};
diff --git a/Engine/Source/Runtime/MovieSceneCapture/Public/Protocols/VideoCaptureProtocol.h b/Engine/Source/Runtime/MovieSceneCapture/Public/Protocols/VideoCaptureProtocol.h
--- a/Engine/Source/Runtime/MovieSceneCapture/Public/Protocols/VideoCaptureProtocol.h
+++ b/Engine/Source/Runtime/MovieSceneCapture/Public/Protocols/VideoCaptureProtocol.h
@@ -11,8 +11,8 @@
 #include "VideoCaptureProtocol.generated.h"
 
 
-UCLASS(meta=(DisplayName="Video Sequence (avi)", CommandLineID="Video"))
-class MOVIESCENECAPTURE_API UVideoCaptureProtocol : public UFrameGrabberProtocol
+UCLASS(meta=(DisplayName="Video Sequence (avi)", CommandLineID="Video"), MinimalAPI)
+class UVideoCaptureProtocol : public UFrameGrabberProtocol
 {
 public:
 	GENERATED_BODY()
@@ -32,15 +32,15 @@ public:
 	float CompressionQuality;
 
 public:
-	virtual bool SetupImpl() override;
-	virtual void FinalizeImpl() override;
-	virtual FFramePayloadPtr GetFramePayload(const FFrameMetrics& FrameMetrics);
-	virtual void ProcessFrame(FCapturedFrameData Frame);
-	virtual bool CanWriteToFileImpl(const TCHAR* InFilename, bool bOverwriteExisting) const override;
+	MOVIESCENECAPTURE_API virtual bool SetupImpl() override;
+	MOVIESCENECAPTURE_API virtual void FinalizeImpl() override;
+	MOVIESCENECAPTURE_API virtual FFramePayloadPtr GetFramePayload(const FFrameMetrics& FrameMetrics);
+	MOVIESCENECAPTURE_API virtual void ProcessFrame(FCapturedFrameData Frame);
+	MOVIESCENECAPTURE_API virtual bool CanWriteToFileImpl(const TCHAR* InFilename, bool bOverwriteExisting) const override;
 
 protected:
 
-	void ConditionallyCreateWriter();
+	MOVIESCENECAPTURE_API void ConditionallyCreateWriter();
 private:
 
 	TArray<TUniquePtr<FAVIWriter>> AVIWriters;
diff --git a/Engine/Source/Runtime/MovieSceneTracks/Private/Evaluation/MovieSceneCameraShakeSourceTriggerTemplate.h b/Engine/Source/Runtime/MovieSceneTracks/Private/Evaluation/MovieSceneCameraShakeSourceTriggerTemplate.h
--- a/Engine/Source/Runtime/MovieSceneTracks/Private/Evaluation/MovieSceneCameraShakeSourceTriggerTemplate.h
+++ b/Engine/Source/Runtime/MovieSceneTracks/Private/Evaluation/MovieSceneCameraShakeSourceTriggerTemplate.h
@@ -12,21 +12,21 @@
 class UMovieSceneCameraShakeSourceTriggerSection;
 
 USTRUCT()
-struct MOVIESCENETRACKS_API FMovieSceneCameraShakeSourceTriggerSectionTemplate : public FMovieSceneEvalTemplate
+struct FMovieSceneCameraShakeSourceTriggerSectionTemplate : public FMovieSceneEvalTemplate
 {
 	GENERATED_BODY()
 	
 	FMovieSceneCameraShakeSourceTriggerSectionTemplate() {}
-	FMovieSceneCameraShakeSourceTriggerSectionTemplate(const UMovieSceneCameraShakeSourceTriggerSection& Section);
+	MOVIESCENETRACKS_API FMovieSceneCameraShakeSourceTriggerSectionTemplate(const UMovieSceneCameraShakeSourceTriggerSection& Section);
 
-	static FMovieSceneAnimTypeID GetAnimTypeID();
+	static MOVIESCENETRACKS_API FMovieSceneAnimTypeID GetAnimTypeID();
 	
 private:
 	virtual UScriptStruct& GetScriptStructImpl() const override { return *StaticStruct(); }
-	virtual void SetupOverrides() override;
-	virtual void Setup(FPersistentEvaluationData& PersistentData, IMovieScenePlayer& Player) const override;
-	virtual void EvaluateSwept(const FMovieSceneEvaluationOperand& Operand, const FMovieSceneContext& Context, const TRange<FFrameNumber>& SweptRange, const FPersistentEvaluationData& PersistentData, FMovieSceneExecutionTokens& ExecutionTokens) const override;
-	virtual void TearDown(FPersistentEvaluationData& PersistentData, IMovieScenePlayer& Player) const override;
+	MOVIESCENETRACKS_API virtual void SetupOverrides() override;
+	MOVIESCENETRACKS_API virtual void Setup(FPersistentEvaluationData& PersistentData, IMovieScenePlayer& Player) const override;
+	MOVIESCENETRACKS_API virtual void EvaluateSwept(const FMovieSceneEvaluationOperand& Operand, const FMovieSceneContext& Context, const TRange<FFrameNumber>& SweptRange, const FPersistentEvaluationData& PersistentData, FMovieSceneExecutionTokens& ExecutionTokens) const override;
+	MOVIESCENETRACKS_API virtual void TearDown(FPersistentEvaluationData& PersistentData, IMovieScenePlayer& Player) const override;
 
 private:
 	UPROPERTY()
diff --git a/Engine/Source/Runtime/MovieSceneTracks/Private/PreAnimatedState/MovieScenePreAnimatedComponentTransformStorage.h b/Engine/Source/Runtime/MovieSceneTracks/Private/PreAnimatedState/MovieScenePreAnimatedComponentTransformStorage.h
--- a/Engine/Source/Runtime/MovieSceneTracks/Private/PreAnimatedState/MovieScenePreAnimatedComponentTransformStorage.h
+++ b/Engine/Source/Runtime/MovieSceneTracks/Private/PreAnimatedState/MovieScenePreAnimatedComponentTransformStorage.h
@@ -23,16 +23,16 @@ struct FComponentTransformPreAnimatedTraits : FComponentTransformPropertyTraits
 	static void SetObjectPropertyValue(UObject* InObject, FTrackInstancePropertyBindings* PropertyBindings, const FIntermediate3DTransform& CachedTransform);
 };
 
-struct MOVIESCENETRACKS_API FPreAnimatedComponentTransformStorage
+struct FPreAnimatedComponentTransformStorage
 	: TPreAnimatedPropertyStorage<FComponentTransformPreAnimatedTraits>
 {
-	FPreAnimatedComponentTransformStorage();
+	MOVIESCENETRACKS_API FPreAnimatedComponentTransformStorage();
 
-	static TAutoRegisterPreAnimatedStorageID<FPreAnimatedComponentTransformStorage> StorageID;
+	static MOVIESCENETRACKS_API TAutoRegisterPreAnimatedStorageID<FPreAnimatedComponentTransformStorage> StorageID;
 
-	void CachePreAnimatedTransforms(const FCachePreAnimatedValueParams& Params, TArrayView<UObject* const> BoundObjects, TOptional<TFunctionRef<bool(int32)>> Predicate = TOptional<TFunctionRef<bool(int32)>>());
+	MOVIESCENETRACKS_API void CachePreAnimatedTransforms(const FCachePreAnimatedValueParams& Params, TArrayView<UObject* const> BoundObjects, TOptional<TFunctionRef<bool(int32)>> Predicate = TOptional<TFunctionRef<bool(int32)>>());
 };
 
 
 } // namespace MovieScene
-} // namespace UE
\ No newline at end of file
+} // namespace UE
diff --git a/Engine/Source/Runtime/MovieSceneTracks/Private/Tests/MovieSceneDecomposerTests.h b/Engine/Source/Runtime/MovieSceneTracks/Private/Tests/MovieSceneDecomposerTests.h
--- a/Engine/Source/Runtime/MovieSceneTracks/Private/Tests/MovieSceneDecomposerTests.h
+++ b/Engine/Source/Runtime/MovieSceneTracks/Private/Tests/MovieSceneDecomposerTests.h
@@ -6,8 +6,8 @@
 
 #include "MovieSceneDecomposerTests.generated.h"
 
-UCLASS()
-class MOVIESCENETRACKS_API UMovieSceneDecomposerTestObject : public UObject
+UCLASS(MinimalAPI)
+class UMovieSceneDecomposerTestObject : public UObject
 {
 	GENERATED_BODY()
 
diff --git a/Engine/Source/Runtime/MovieSceneTracks/Public/Channels/MovieSceneCameraShakeSourceTriggerChannel.h b/Engine/Source/Runtime/MovieSceneTracks/Public/Channels/MovieSceneCameraShakeSourceTriggerChannel.h
--- a/Engine/Source/Runtime/MovieSceneTracks/Public/Channels/MovieSceneCameraShakeSourceTriggerChannel.h
+++ b/Engine/Source/Runtime/MovieSceneTracks/Public/Channels/MovieSceneCameraShakeSourceTriggerChannel.h
@@ -59,7 +59,7 @@ struct FMovieSceneCameraShakeSourceTrigger
 };
 
 USTRUCT()
-struct MOVIESCENETRACKS_API FMovieSceneCameraShakeSourceTriggerChannel : public FMovieSceneChannel
+struct FMovieSceneCameraShakeSourceTriggerChannel : public FMovieSceneChannel
 {
 	GENERATED_BODY()
 
@@ -75,17 +75,17 @@ struct MOVIESCENETRACKS_API FMovieSceneCameraShakeSourceTriggerChannel : public
 
 public:
 	// ~ FMovieSceneChannel Interface
-	virtual void GetKeys(const TRange<FFrameNumber>& WithinRange, TArray<FFrameNumber>* OutKeyTimes, TArray<FKeyHandle>* OutKeyHandles) override;
-	virtual void GetKeyTimes(TArrayView<const FKeyHandle> InHandles, TArrayView<FFrameNumber> OutKeyTimes) override;
-	virtual void SetKeyTimes(TArrayView<const FKeyHandle> InHandles, TArrayView<const FFrameNumber> InKeyTimes) override;
-	virtual void DuplicateKeys(TArrayView<const FKeyHandle> InHandles, TArrayView<FKeyHandle> OutNewHandles) override;
-	virtual void DeleteKeys(TArrayView<const FKeyHandle> InHandles) override;
-	virtual void DeleteKeysFrom(FFrameNumber InTime, bool bDeleteKeysBefore) override;
-	virtual void ChangeFrameResolution(FFrameRate SourceRate, FFrameRate DestinationRate) override;
-	virtual TRange<FFrameNumber> ComputeEffectiveRange() const override;
-	virtual int32 GetNumKeys() const override;
-	virtual void Reset() override;
-	virtual void Offset(FFrameNumber DeltaPosition) override;
+	MOVIESCENETRACKS_API virtual void GetKeys(const TRange<FFrameNumber>& WithinRange, TArray<FFrameNumber>* OutKeyTimes, TArray<FKeyHandle>* OutKeyHandles) override;
+	MOVIESCENETRACKS_API virtual void GetKeyTimes(TArrayView<const FKeyHandle> InHandles, TArrayView<FFrameNumber> OutKeyTimes) override;
+	MOVIESCENETRACKS_API virtual void SetKeyTimes(TArrayView<const FKeyHandle> InHandles, TArrayView<const FFrameNumber> InKeyTimes) override;
+	MOVIESCENETRACKS_API virtual void DuplicateKeys(TArrayView<const FKeyHandle> InHandles, TArrayView<FKeyHandle> OutNewHandles) override;
+	MOVIESCENETRACKS_API virtual void DeleteKeys(TArrayView<const FKeyHandle> InHandles) override;
+	MOVIESCENETRACKS_API virtual void DeleteKeysFrom(FFrameNumber InTime, bool bDeleteKeysBefore) override;
+	MOVIESCENETRACKS_API virtual void ChangeFrameResolution(FFrameRate SourceRate, FFrameRate DestinationRate) override;
+	MOVIESCENETRACKS_API virtual TRange<FFrameNumber> ComputeEffectiveRange() const override;
+	MOVIESCENETRACKS_API virtual int32 GetNumKeys() const override;
+	MOVIESCENETRACKS_API virtual void Reset() override;
+	MOVIESCENETRACKS_API virtual void Offset(FFrameNumber DeltaPosition) override;
 
 private:
 	/** Array of times for each key */
diff --git a/Engine/Source/Runtime/MovieSceneTracks/Public/Channels/MovieSceneDoublePerlinNoiseChannel.h b/Engine/Source/Runtime/MovieSceneTracks/Public/Channels/MovieSceneDoublePerlinNoiseChannel.h
--- a/Engine/Source/Runtime/MovieSceneTracks/Public/Channels/MovieSceneDoublePerlinNoiseChannel.h
+++ b/Engine/Source/Runtime/MovieSceneTracks/Public/Channels/MovieSceneDoublePerlinNoiseChannel.h
@@ -8,13 +8,13 @@
 #include "MovieSceneDoublePerlinNoiseChannel.generated.h"
 
 USTRUCT()
-struct MOVIESCENETRACKS_API FMovieSceneDoublePerlinNoiseChannel : public FMovieSceneChannel
+struct FMovieSceneDoublePerlinNoiseChannel : public FMovieSceneChannel
 {
 	GENERATED_BODY()
 
 	using CurveValueType = double;
 
-	FMovieSceneDoublePerlinNoiseChannel();
+	MOVIESCENETRACKS_API FMovieSceneDoublePerlinNoiseChannel();
 
 	/**
 	 * Evaluate this channel at the given time
@@ -22,7 +22,7 @@ struct MOVIESCENETRACKS_API FMovieSceneDoublePerlinNoiseChannel : public FMovieS
 	 * @param InSeconds		The time, in seconds, to evaluate at
 	 * @return				The evaluated noise value
 	 */
-	double Evaluate(double InSeconds) const;
+	MOVIESCENETRACKS_API double Evaluate(double InSeconds) const;
 
 	/**
 	 * Evaluate this channel at the given time
@@ -32,7 +32,7 @@ struct MOVIESCENETRACKS_API FMovieSceneDoublePerlinNoiseChannel : public FMovieS
 	 * @param OutValue		The evaluated noise value
 	 * @return				Whether the noise was successfully evaluated
 	 */
-	bool Evaluate(const UMovieSceneSection* InSection, FFrameTime InTime, double& OutValue) const;
+	MOVIESCENETRACKS_API bool Evaluate(const UMovieSceneSection* InSection, FFrameTime InTime, double& OutValue) const;
 
 	/** The noise parameters */
 	UPROPERTY(EditAnywhere, Category = "Perlin Noise")
@@ -46,7 +46,7 @@ struct MOVIESCENETRACKS_API FMovieSceneDoublePerlinNoiseChannel : public FMovieS
 	 * @params InSeconds		The time at which to evaluate the noise
 	 * @return					The evaluated noise value
 	 */
-	static double Evaluate(const FPerlinNoiseParams& InParams, double InSeconds);
+	static MOVIESCENETRACKS_API double Evaluate(const FPerlinNoiseParams& InParams, double InSeconds);
 };
 
 template<>
diff --git a/Engine/Source/Runtime/MovieSceneTracks/Public/Channels/MovieSceneDoublePerlinNoiseChannelContainer.h b/Engine/Source/Runtime/MovieSceneTracks/Public/Channels/MovieSceneDoublePerlinNoiseChannelContainer.h
--- a/Engine/Source/Runtime/MovieSceneTracks/Public/Channels/MovieSceneDoublePerlinNoiseChannelContainer.h
+++ b/Engine/Source/Runtime/MovieSceneTracks/Public/Channels/MovieSceneDoublePerlinNoiseChannelContainer.h
@@ -9,8 +9,8 @@
 /**
  * Double perlin noise channel overriden container
  */
-UCLASS(meta=(DisplayName="Double Perlin Noise", ToolTip="Override a channel to use double perlin noise"))
-class  MOVIESCENETRACKS_API UMovieSceneDoublePerlinNoiseChannelContainer : public UMovieSceneChannelOverrideContainer
+UCLASS(meta=(DisplayName="Double Perlin Noise", ToolTip="Override a channel to use double perlin noise"), MinimalAPI)
+class  UMovieSceneDoublePerlinNoiseChannelContainer : public UMovieSceneChannelOverrideContainer
 {
 	GENERATED_BODY()
 
@@ -18,9 +18,9 @@ public:
 
 	using ChannelType = FMovieSceneDoublePerlinNoiseChannel;
 
-	void InitializeOverride(FMovieSceneChannel* InChannel) override;
-	bool SupportsOverride(FName DefaultChannelTypeName) const override;
-	void ImportEntityImpl(
+	MOVIESCENETRACKS_API void InitializeOverride(FMovieSceneChannel* InChannel) override;
+	MOVIESCENETRACKS_API bool SupportsOverride(FName DefaultChannelTypeName) const override;
+	MOVIESCENETRACKS_API void ImportEntityImpl(
 			const UE::MovieScene::FChannelOverrideEntityImportParams& OverrideParams, 
 			const UE::MovieScene::FEntityImportParams& ImportParams, UE::MovieScene::FImportedEntity* OutImportedEntity) override;
 
@@ -28,9 +28,9 @@ public:
 	FMovieSceneChannel* GetChannel() override { return &PerlinNoiseChannel; }
 
 #if WITH_EDITOR
-	FMovieSceneChannelHandle AddChannelProxy(FName ChannelName, FMovieSceneChannelProxyData& ProxyData, const FMovieSceneChannelMetaData& MetaData) override;
+	MOVIESCENETRACKS_API FMovieSceneChannelHandle AddChannelProxy(FName ChannelName, FMovieSceneChannelProxyData& ProxyData, const FMovieSceneChannelMetaData& MetaData) override;
 #else
-	void AddChannelProxy(FName ChannelName, FMovieSceneChannelProxyData& ProxyData) override;
+	MOVIESCENETRACKS_API void AddChannelProxy(FName ChannelName, FMovieSceneChannelProxyData& ProxyData) override;
 #endif
 
 private:
diff --git a/Engine/Source/Runtime/MovieSceneTracks/Public/Channels/MovieSceneEvent.h b/Engine/Source/Runtime/MovieSceneTracks/Public/Channels/MovieSceneEvent.h
--- a/Engine/Source/Runtime/MovieSceneTracks/Public/Channels/MovieSceneEvent.h
+++ b/Engine/Source/Runtime/MovieSceneTracks/Public/Channels/MovieSceneEvent.h
@@ -56,7 +56,7 @@ struct FMovieSceneEventPtrs
 };
 
 USTRUCT(BlueprintType)
-struct MOVIESCENETRACKS_API FMovieSceneEvent
+struct FMovieSceneEvent
 {
 	GENERATED_BODY()
 
@@ -70,7 +70,7 @@ struct MOVIESCENETRACKS_API FMovieSceneEvent
 public:
 
 	/** Return the class of the bound object property */
-	UClass* GetBoundObjectPropertyClass() const;
+	MOVIESCENETRACKS_API UClass* GetBoundObjectPropertyClass() const;
 
 #if WITH_EDITORONLY_DATA
 
