Fix FFrameGrabber in d3d12. When using Movie Scene Capture (Legacy), the first few frames captured contained nothing or garbage. Added a GPU fence and had the call to MapStagingSurface to wait on it so the GPU would finish writing to the readback texture before being used to write images to the output folder.

diff --git a/Engine/Source/Runtime/MovieSceneCapture/Private/FrameGrabber.cpp b/Engine/Source/Runtime/MovieSceneCapture/Private/FrameGrabber.cpp
--- a/Engine/Source/Runtime/MovieSceneCapture/Private/FrameGrabber.cpp
+++ b/Engine/Source/Runtime/MovieSceneCapture/Private/FrameGrabber.cpp
@@ -169,14 +169,16 @@ void FViewportSurfaceReader::ResolveRenderTarget(FViewportSurfaceReader* RenderT
 				EDRF_Default);
 		}
 		RHICmdList.EndRenderPass();
+		FGPUFenceRHIRef GPUFence = RHICreateGPUFence(TEXT("FrameGrabberReadbackFence"));
 		TransitionAndCopyTexture(RHICmdList, DestRenderTarget, ReadbackTexture, {});
+		RHICmdList.WriteGPUFence(GPUFence);
 
 		if (RenderToReadback)
 		{
 			void* ColorDataBuffer = nullptr;
 
 			int32 Width = 0, Height = 0;
-			RHICmdList.MapStagingSurface(RenderToReadback->ReadbackTexture, ColorDataBuffer, Width, Height);
+			RHICmdList.MapStagingSurface(RenderToReadback->ReadbackTexture, GPUFence.GetReference(), ColorDataBuffer, Width, Height);
 
 			Callback((FColor*)ColorDataBuffer, Width, Height);
 
