Movie Scene Capture: Speculativ0 nullptre fix

#rnx

diff --git a/Engine/Source/Runtime/MovieSceneCapture/Private/MovieSceneCaptureProtocolBase.cpp b/Engine/Source/Runtime/MovieSceneCapture/Private/MovieSceneCaptureProtocolBase.cpp
--- a/Engine/Source/Runtime/MovieSceneCapture/Private/MovieSceneCaptureProtocolBase.cpp
+++ b/Engine/Source/Runtime/MovieSceneCapture/Private/MovieSceneCaptureProtocolBase.cpp
@@ -52,7 +52,7 @@ bool UMovieSceneCaptureProtocolBase::Setup(const FCaptureProtocolInitSettings& I
 
 UWorld* UMovieSceneCaptureProtocolBase::GetWorld() const
 {
-	if (InitSettings.IsSet())
+	if (InitSettings.IsSet() && InitSettings->SceneViewport != nullptr && InitSettings->SceneViewport->GetClient() != nullptr)
 	{
 		// Retrieve the world from the Scene Viewport client.
 		return InitSettings->SceneViewport->GetClient()->GetWorld();
